UE4 - Fix deprecation and static analysis warnings
#rnx

diff --git a/Engine/Plugins/Experimental/ChaosNiagara/Source/ChaosNiagara/Private/NiagaraDataInterfaceChaosDestruction.cpp b/Engine/Plugins/Experimental/ChaosNiagara/Source/ChaosNiagara/Private/NiagaraDataInterfaceChaosDestruction.cpp
--- a/Engine/Plugins/Experimental/ChaosNiagara/Source/ChaosNiagara/Private/NiagaraDataInterfaceChaosDestruction.cpp
+++ b/Engine/Plugins/Experimental/ChaosNiagara/Source/ChaosNiagara/Private/NiagaraDataInterfaceChaosDestruction.cpp
@@ -4090,9 +4090,9 @@ void LoadGPUBufferFromArray(FDynamicReadBuffer& Buffer,
 	Buffer.Unlock();
 }
 
-void SetBuffer(FRHICommandList& CmdList,
+static void SetBuffer(FRHICommandList& CmdList,
 	const FShaderResourceParameter& Param,
-	const FComputeShaderRHIParamRef Shader,
+	FRHIComputeShader* Shader,
 	FDynamicReadBuffer& Buffer)
 {
 	// Skip unbound parameters, since we won't be reading any of them
@@ -4102,9 +4102,9 @@ void SetBuffer(FRHICommandList& CmdList,
 }
 
 template <typename T>
-void SetBuffer(FRHICommandList& CmdList,
+static void SetBuffer(FRHICommandList& CmdList,
 	const FShaderResourceParameter& Param,
-	const FComputeShaderRHIParamRef Shader,
+	FRHIComputeShader* Shader,
 	FDynamicReadBuffer& Buffer,
 	const TArray<T>& Array,
 	const EPixelFormat PixelFormat,
@@ -4605,7 +4605,7 @@ struct FNiagaraDataInterfaceParametersCS_ChaosDestruction : public FNiagaraDataI
 	{
 		check(IsInRenderingThread());
 
-		const FComputeShaderRHIParamRef ComputeShaderRHI = Context.Shader->GetComputeShader();
+		FRHIComputeShader* ComputeShaderRHI = Context.Shader->GetComputeShader();
 		FNiagaraDataInterfaceProxyChaosDestruction* ChaosDestructionInterfaceProxy = static_cast<FNiagaraDataInterfaceProxyChaosDestruction*>(Context.DataInterface);
 		if (ChaosDestructionInterfaceProxy)
 		{
diff --git a/Engine/Source/Runtime/VulkanRHI/Private/VulkanViewport.cpp b/Engine/Source/Runtime/VulkanRHI/Private/VulkanViewport.cpp
--- a/Engine/Source/Runtime/VulkanRHI/Private/VulkanViewport.cpp
+++ b/Engine/Source/Runtime/VulkanRHI/Private/VulkanViewport.cpp
@@ -522,13 +522,9 @@ void FVulkanViewport::RecreateSwapchain(void* NewNativeWindow, bool bForce)
 		for (int32 Index = 0; Index < NUM_BUFFERS; ++Index)
 		{
 			TextureViews[Index].Destroy(*Device);
-		}
-
-		for (VkImage& BackBufferImage : BackBufferImages)
-		{
-			Device->NotifyDeletedImage(BackBufferImage);
-			Device->NotifyDeletedRenderTarget(BackBufferImage);
-			BackBufferImage = VK_NULL_HANDLE;
+			Device->NotifyDeletedImage(BackBufferImages[Index]);
+			Device->NotifyDeletedRenderTarget(BackBufferImages[Index]);
+			BackBufferImages[Index] = VK_NULL_HANDLE;
 		}
 
 		Device->GetDeferredDeletionQueue().ReleaseResources(true);
@@ -591,13 +587,9 @@ void FVulkanViewport::RecreateSwapchainFromRT(EPixelFormat PreferredPixelFormat)
 		for (int32 Index = 0; Index < NUM_BUFFERS; ++Index)
 		{
 			TextureViews[Index].Destroy(*Device);
-		}
-		
-		for (VkImage& BackBufferImage : BackBufferImages)
-		{
-			Device->NotifyDeletedImage(BackBufferImage);
-			Device->NotifyDeletedRenderTarget(BackBufferImage);
-			BackBufferImage = VK_NULL_HANDLE;
+			Device->NotifyDeletedImage(BackBufferImages[Index]);
+			Device->NotifyDeletedRenderTarget(BackBufferImages[Index]);
+			BackBufferImages[Index] = VK_NULL_HANDLE;
 		}
 		
 		Device->GetDeferredDeletionQueue().ReleaseResources(true);
diff --git a/Engine/Source/Runtime/VulkanRHI/Private/VulkanViewport.h b/Engine/Source/Runtime/VulkanRHI/Private/VulkanViewport.h
--- a/Engine/Source/Runtime/VulkanRHI/Private/VulkanViewport.h
+++ b/Engine/Source/Runtime/VulkanRHI/Private/VulkanViewport.h
@@ -81,9 +81,9 @@ public:
 	}
 
 protected:
-	VkImage BackBufferImages[NUM_BUFFERS];
-	VulkanRHI::FSemaphore* RenderingDoneSemaphores[NUM_BUFFERS];
-	FVulkanTextureView TextureViews[NUM_BUFFERS];
+	TStaticArray<VkImage, NUM_BUFFERS> BackBufferImages;
+	TStaticArray<VulkanRHI::FSemaphore*, NUM_BUFFERS> RenderingDoneSemaphores;
+	TStaticArray<FVulkanTextureView, NUM_BUFFERS> TextureViews;
 	TRefCountPtr<FVulkanBackBuffer> RHIBackBuffer;
 
 	// 'Dummy' back buffer
