- Fix modules I missed with FGuid to FNiagaraSystemInstanceID conversion

#rnx

diff --git a/Engine/Plugins/Experimental/ChaosNiagara/Source/ChaosNiagara/Classes/NiagaraDataInterfaceChaosDestruction.h b/Engine/Plugins/Experimental/ChaosNiagara/Source/ChaosNiagara/Classes/NiagaraDataInterfaceChaosDestruction.h
--- a/Engine/Plugins/Experimental/ChaosNiagara/Source/ChaosNiagara/Classes/NiagaraDataInterfaceChaosDestruction.h
+++ b/Engine/Plugins/Experimental/ChaosNiagara/Source/ChaosNiagara/Classes/NiagaraDataInterfaceChaosDestruction.h
@@ -772,13 +772,13 @@ struct FNiagaraDataInterfaceProxyChaosDestruction : public FNiagaraDataInterface
 		return sizeof(FNiagaraDIChaosDestruction_InstanceDataToPassToRT);
 	}
 
-	void CreatePerInstanceData(const FGuid& SystemInstance);
+	void CreatePerInstanceData(const FNiagaraSystemInstanceID& SystemInstance);
 
-	void DestroyInstanceData(NiagaraEmitterInstanceBatcher* Batcher, const FGuid& SystemInstance);
+	void DestroyInstanceData(NiagaraEmitterInstanceBatcher* Batcher, const FNiagaraSystemInstanceID& SystemInstance);
 
 	virtual void DeferredDestroy()
 	{
-		for (const FGuid& SystemInstance : InstancesToDestroy)
+		for (const FNiagaraSystemInstanceID& SystemInstance : InstancesToDestroy)
 		{
 			SystemsToGPUInstanceData.Remove(SystemInstance);
 		}
@@ -789,6 +789,6 @@ struct FNiagaraDataInterfaceProxyChaosDestruction : public FNiagaraDataInterface
 	float SolverTime;
 	int32 LastSpawnedPointID;
 
-	TMap<FGuid, FNiagaraDIChaosDestruction_GPUData> SystemsToGPUInstanceData;
-	TSet<FGuid> InstancesToDestroy;
+	TMap<FNiagaraSystemInstanceID, FNiagaraDIChaosDestruction_GPUData> SystemsToGPUInstanceData;
+	TSet<FNiagaraSystemInstanceID> InstancesToDestroy;
 };
\ No newline at end of file
diff --git a/Engine/Plugins/Experimental/ChaosNiagara/Source/ChaosNiagara/Private/NiagaraDataInterfaceChaosDestruction.cpp b/Engine/Plugins/Experimental/ChaosNiagara/Source/ChaosNiagara/Private/NiagaraDataInterfaceChaosDestruction.cpp
--- a/Engine/Plugins/Experimental/ChaosNiagara/Source/ChaosNiagara/Private/NiagaraDataInterfaceChaosDestruction.cpp
+++ b/Engine/Plugins/Experimental/ChaosNiagara/Source/ChaosNiagara/Private/NiagaraDataInterfaceChaosDestruction.cpp
@@ -4085,7 +4085,7 @@ void UNiagaraDataInterfaceChaosDestruction::ProvidePerInstanceDataForRenderThrea
 	}
 }
 
-void FNiagaraDataInterfaceProxyChaosDestruction::CreatePerInstanceData(const FGuid& SystemInstance)
+void FNiagaraDataInterfaceProxyChaosDestruction::CreatePerInstanceData(const FNiagaraSystemInstanceID& SystemInstance)
 {
 	check(IsInRenderingThread());
 	if (SystemsToGPUInstanceData.Contains(SystemInstance))
diff --git a/Engine/Plugins/Experimental/HairStrands/Source/HairStrandsNiagara/Private/NiagaraDataInterfaceHairStrands.cpp b/Engine/Plugins/Experimental/HairStrands/Source/HairStrandsNiagara/Private/NiagaraDataInterfaceHairStrands.cpp
--- a/Engine/Plugins/Experimental/HairStrands/Source/HairStrandsNiagara/Private/NiagaraDataInterfaceHairStrands.cpp
+++ b/Engine/Plugins/Experimental/HairStrands/Source/HairStrandsNiagara/Private/NiagaraDataInterfaceHairStrands.cpp
@@ -268,7 +268,7 @@ private:
 
 void FNDIHairStrandsProxy::DeferredDestroy()
 {
-	for (const FGuid& Sys : DeferredDestroyList)
+	for (const FNiagaraSystemInstanceID& Sys : DeferredDestroyList)
 	{
 		SystemInstancesToProxyData.Remove(Sys);
 	}
@@ -288,7 +288,7 @@ void FNDIHairStrandsProxy::ConsumePerInstanceDataFromGameThread(void* PerInstanc
 	}
 	else
 	{
-		UE_LOG(LogHairStrands, Log, TEXT("ConsumePerInstanceDataFromGameThread() ... could not find %s"), *Instance.ToString());
+		UE_LOG(LogHairStrands, Log, TEXT("ConsumePerInstanceDataFromGameThread() ... could not find %s"), *FNiagaraUtilities::SystemInstanceIDToString(Instance));
 	}
 }
 
@@ -314,7 +314,7 @@ void FNDIHairStrandsProxy::InitializePerInstanceData(const FNiagaraSystemInstanc
 	TargetData->TipThickness = TipThickness;
 }
 
-void FNDIHairStrandsProxy::DestroyPerInstanceData(NiagaraEmitterInstanceBatcher* Batcher, const FGuid& SystemInstance)
+void FNDIHairStrandsProxy::DestroyPerInstanceData(NiagaraEmitterInstanceBatcher* Batcher, const FNiagaraSystemInstanceID& SystemInstance)
 {
 	check(IsInRenderingThread());
 
diff --git a/Engine/Plugins/Experimental/HairStrands/Source/HairStrandsNiagara/Public/NiagaraDataInterfaceHairStrands.h b/Engine/Plugins/Experimental/HairStrands/Source/HairStrandsNiagara/Public/NiagaraDataInterfaceHairStrands.h
--- a/Engine/Plugins/Experimental/HairStrands/Source/HairStrandsNiagara/Public/NiagaraDataInterfaceHairStrands.h
+++ b/Engine/Plugins/Experimental/HairStrands/Source/HairStrandsNiagara/Public/NiagaraDataInterfaceHairStrands.h
@@ -227,12 +227,12 @@ struct FNDIHairStrandsProxy : public FNiagaraDataInterfaceProxy
 		const float StrandDensity, const float RootThickness, const float TipThickness);
 
 	/** Destroy the proxy data if necessary */
-	void DestroyPerInstanceData(NiagaraEmitterInstanceBatcher* Batcher, const FGuid& SystemInstance);
+	void DestroyPerInstanceData(NiagaraEmitterInstanceBatcher* Batcher, const FNiagaraSystemInstanceID& SystemInstance);
 
 	/** List of proxy data for each system instances*/
-	TMap<FGuid, FNDIHairStrandsData> SystemInstancesToProxyData;
+	TMap<FNiagaraSystemInstanceID, FNDIHairStrandsData> SystemInstancesToProxyData;
 
 	/** List of proxy data to destroy later */
-	TSet<FGuid> DeferredDestroyList;
+	TSet<FNiagaraSystemInstanceID> DeferredDestroyList;
 };
 
