- Remove deferred deletion from Data Interfaces it is not required
- Remove TSharedPtr from DI proxy as it is no longer required

diff --git a/Engine/Plugins/Experimental/ChaosNiagara/Source/ChaosNiagara/Classes/NiagaraDataInterfaceChaosDestruction.h b/Engine/Plugins/Experimental/ChaosNiagara/Source/ChaosNiagara/Classes/NiagaraDataInterfaceChaosDestruction.h
--- a/Engine/Plugins/Experimental/ChaosNiagara/Source/ChaosNiagara/Classes/NiagaraDataInterfaceChaosDestruction.h
+++ b/Engine/Plugins/Experimental/ChaosNiagara/Source/ChaosNiagara/Classes/NiagaraDataInterfaceChaosDestruction.h
@@ -764,19 +764,8 @@ struct FNiagaraDataInterfaceProxyChaosDestruction : public FNiagaraDataInterface
 
 	void DestroyInstanceData(NiagaraEmitterInstanceBatcher* Batcher, const FNiagaraSystemInstanceID& SystemInstance);
 
-	virtual void DeferredDestroy()
-	{
-		for (const FNiagaraSystemInstanceID& SystemInstance : InstancesToDestroy)
-		{
-			SystemsToGPUInstanceData.Remove(SystemInstance);
-		}
-
-		InstancesToDestroy.Empty();
-	}
-
 	float SolverTime;
 	int32 LastSpawnedPointID;
 
 	TMap<FNiagaraSystemInstanceID, FNiagaraDIChaosDestruction_GPUData> SystemsToGPUInstanceData;
-	TSet<FNiagaraSystemInstanceID> InstancesToDestroy;
 };
\ No newline at end of file
diff --git a/Engine/Plugins/Experimental/ChaosNiagara/Source/ChaosNiagara/Private/NiagaraDataInterfaceChaosDestruction.cpp b/Engine/Plugins/Experimental/ChaosNiagara/Source/ChaosNiagara/Private/NiagaraDataInterfaceChaosDestruction.cpp
--- a/Engine/Plugins/Experimental/ChaosNiagara/Source/ChaosNiagara/Private/NiagaraDataInterfaceChaosDestruction.cpp
+++ b/Engine/Plugins/Experimental/ChaosNiagara/Source/ChaosNiagara/Private/NiagaraDataInterfaceChaosDestruction.cpp
@@ -142,7 +142,7 @@ UNiagaraDataInterfaceChaosDestruction::UNiagaraDataInterfaceChaosDestruction(FOb
 
 	Solvers.Reset();
 
-	Proxy = MakeShared<FNiagaraDataInterfaceProxyChaosDestruction, ESPMode::ThreadSafe>();
+	Proxy.Reset(new FNiagaraDataInterfaceProxyChaosDestruction());
 	PushToRenderThread();
 }
 
@@ -483,7 +483,7 @@ bool UNiagaraDataInterfaceChaosDestruction::InitPerInstanceData(void* PerInstanc
 
 	ResetInstData(InstData);
 
-	TSharedPtr<FNiagaraDataInterfaceProxyChaosDestruction, ESPMode::ThreadSafe> ThisProxy = StaticCastSharedPtr<FNiagaraDataInterfaceProxyChaosDestruction, FNiagaraDataInterfaceProxy, ESPMode::ThreadSafe>(Proxy);
+	FNiagaraDataInterfaceProxyChaosDestruction* ThisProxy = GetProxyAs<FNiagaraDataInterfaceProxyChaosDestruction>();
 	ENQUEUE_RENDER_COMMAND(FNiagaraChaosDestructionDICreateRTInstance)(
 		[ThisProxy, InstanceID = SystemInstance->GetId()](FRHICommandList& CmdList)
 	{
@@ -500,7 +500,7 @@ void UNiagaraDataInterfaceChaosDestruction::DestroyPerInstanceData(void* PerInst
 	InstData->~FNDIChaosDestruction_InstanceData();
 
 	check(Proxy);
-	TSharedPtr<FNiagaraDataInterfaceProxyChaosDestruction, ESPMode::ThreadSafe> ThisProxy = StaticCastSharedPtr<FNiagaraDataInterfaceProxyChaosDestruction, FNiagaraDataInterfaceProxy, ESPMode::ThreadSafe>(Proxy);
+	FNiagaraDataInterfaceProxyChaosDestruction* ThisProxy = GetProxyAs<FNiagaraDataInterfaceProxyChaosDestruction>();
 	ENQUEUE_RENDER_COMMAND(FNiagaraDIChaosDestructionDestroyInstanceData) (
 		[ThisProxy, InstanceID = SystemInstance->GetId(), Batcher = SystemInstance->GetBatcher()](FRHICommandListImmediate& CmdList)
 		{
@@ -3923,7 +3923,7 @@ static void SetBuffer(FRHICommandList& CmdList,
 void UNiagaraDataInterfaceChaosDestruction::PushToRenderThread()
 {
 	check(Proxy);
-	TSharedPtr<FNiagaraDataInterfaceProxyChaosDestruction, ESPMode::ThreadSafe> RT_Proxy = StaticCastSharedPtr<FNiagaraDataInterfaceProxyChaosDestruction, FNiagaraDataInterfaceProxy, ESPMode::ThreadSafe>(Proxy);
+	FNiagaraDataInterfaceProxyChaosDestruction* RT_Proxy = GetProxyAs<FNiagaraDataInterfaceProxyChaosDestruction>();
 
 	int32 RT_LastSpawnedPointID = LastSpawnedPointID;
 	float RT_SolverTime = SolverTime;
@@ -4077,21 +4077,15 @@ void UNiagaraDataInterfaceChaosDestruction::ProvidePerInstanceDataForRenderThrea
 void FNiagaraDataInterfaceProxyChaosDestruction::CreatePerInstanceData(const FNiagaraSystemInstanceID& SystemInstance)
 {
 	check(IsInRenderingThread());
-	if (SystemsToGPUInstanceData.Contains(SystemInstance))
-	{
-		InstancesToDestroy.Remove(SystemInstance);
-	}
+	check(!SystemsToGPUInstanceData.Contains(SystemInstance));
 	SystemsToGPUInstanceData.Add(SystemInstance, FNiagaraDIChaosDestruction_GPUData());
 }
 
 void FNiagaraDataInterfaceProxyChaosDestruction::DestroyInstanceData(NiagaraEmitterInstanceBatcher* Batcher, const FNiagaraSystemInstanceID& SystemInstance)
 {
 	check(IsInRenderingThread());
-	// @todo-threadsafety This object contains GPU buffers. This _should_ delete them safely but would we rather do so manually?
-	//SystemsToGPUInstanceData.Remove(SystemInstance);
-	InstancesToDestroy.Add(SystemInstance);
-
-	Batcher->EnqueueDeferredDeletesForDI_RenderThread(this->AsShared());
+	check(SystemsToGPUInstanceData.Contains(SystemInstance));
+	SystemsToGPUInstanceData.Remove(SystemInstance);
 }
 
 void FNiagaraDataInterfaceProxyChaosDestruction::ConsumePerInstanceDataFromGameThread(void* PerInstanceDataFromGameThread, const FNiagaraSystemInstanceID& Instance)
diff --git a/Engine/Plugins/Experimental/HairStrands/Source/HairStrandsNiagara/Private/NiagaraDataInterfaceHairStrands.cpp b/Engine/Plugins/Experimental/HairStrands/Source/HairStrandsNiagara/Private/NiagaraDataInterfaceHairStrands.cpp
--- a/Engine/Plugins/Experimental/HairStrands/Source/HairStrandsNiagara/Private/NiagaraDataInterfaceHairStrands.cpp
+++ b/Engine/Plugins/Experimental/HairStrands/Source/HairStrandsNiagara/Private/NiagaraDataInterfaceHairStrands.cpp
@@ -843,16 +843,6 @@ private:
 
 //------------------------------------------------------------------------------------------------------------
 
-void FNDIHairStrandsProxy::DeferredDestroy()
-{
-	for (const FNiagaraSystemInstanceID& Sys : DeferredDestroyList)
-	{
-		SystemInstancesToProxyData.Remove(Sys);
-	}
-
-	DeferredDestroyList.Empty();
-}
-
 void FNDIHairStrandsProxy::ConsumePerInstanceDataFromGameThread(void* PerInstanceData, const FNiagaraSystemInstanceID& Instance)
 {
 	FNDIHairStrandsData* SourceData = static_cast<FNDIHairStrandsData*>(PerInstanceData);
@@ -882,22 +872,14 @@ void FNDIHairStrandsProxy::InitializePerInstanceData(const FNiagaraSystemInstanc
 	check(IsInRenderingThread());
 
 	FNDIHairStrandsData* TargetData = SystemInstancesToProxyData.Find(SystemInstance);
-	if (TargetData != nullptr)
-	{
-		DeferredDestroyList.Remove(SystemInstance);
-	}
-	else
-	{
-		TargetData = &SystemInstancesToProxyData.Add(SystemInstance);
-	}
+	TargetData = &SystemInstancesToProxyData.Add(SystemInstance);
 }
 
 void FNDIHairStrandsProxy::DestroyPerInstanceData(NiagaraEmitterInstanceBatcher* Batcher, const FNiagaraSystemInstanceID& SystemInstance)
 {
 	check(IsInRenderingThread());
-
-	DeferredDestroyList.Add(SystemInstance);
-	Batcher->EnqueueDeferredDeletesForDI_RenderThread(this->AsShared());
+	//check(SystemInstancesToProxyData.Contains(SystemInstance));
+	SystemInstancesToProxyData.Remove(SystemInstance);
 }
 
 //------------------------------------------------------------------------------------------------------------
@@ -911,7 +893,7 @@ UNiagaraDataInterfaceHairStrands::UNiagaraDataInterfaceHairStrands(FObjectInitia
 	, GroupIndex(0)
 {
 
-	Proxy = MakeShared<FNDIHairStrandsProxy, ESPMode::ThreadSafe>();
+	Proxy.Reset(new FNDIHairStrandsProxy());
 }
 
 bool UNiagaraDataInterfaceHairStrands::IsComponentValid() const
diff --git a/Engine/Plugins/Experimental/HairStrands/Source/HairStrandsNiagara/Private/NiagaraDataInterfacePhysicsAsset.cpp b/Engine/Plugins/Experimental/HairStrands/Source/HairStrandsNiagara/Private/NiagaraDataInterfacePhysicsAsset.cpp
--- a/Engine/Plugins/Experimental/HairStrands/Source/HairStrandsNiagara/Private/NiagaraDataInterfacePhysicsAsset.cpp
+++ b/Engine/Plugins/Experimental/HairStrands/Source/HairStrandsNiagara/Private/NiagaraDataInterfacePhysicsAsset.cpp
@@ -502,16 +502,6 @@ private:
 
 //------------------------------------------------------------------------------------------------------------
 
-void FNDIPhysicsAssetProxy::DeferredDestroy()
-{
-	for (const FNiagaraSystemInstanceID& Sys : DeferredDestroyList)
-	{
-		SystemInstancesToProxyData.Remove(Sys);
-	}
-
-	DeferredDestroyList.Empty();
-}
-
 void FNDIPhysicsAssetProxy::ConsumePerInstanceDataFromGameThread(void* PerInstanceData, const FNiagaraSystemInstanceID& Instance)
 {
 	FNDIPhysicsAssetData* SourceData = static_cast<FNDIPhysicsAssetData*>(PerInstanceData);
@@ -535,22 +525,14 @@ void FNDIPhysicsAssetProxy::InitializePerInstanceData(const FNiagaraSystemInstan
 	check(IsInRenderingThread());
 
 	FNDIPhysicsAssetData* TargetData = SystemInstancesToProxyData.Find(SystemInstance);
-	if (TargetData != nullptr)
-	{
-		DeferredDestroyList.Remove(SystemInstance);
-	}
-	else
-	{
-		TargetData = &SystemInstancesToProxyData.Add(SystemInstance);
-	}
+	TargetData = &SystemInstancesToProxyData.Add(SystemInstance);
 }
 
 void FNDIPhysicsAssetProxy::DestroyPerInstanceData(NiagaraEmitterInstanceBatcher* Batcher, const FNiagaraSystemInstanceID& SystemInstance)
 {
 	check(IsInRenderingThread());
-
-	DeferredDestroyList.Add(SystemInstance);
-	Batcher->EnqueueDeferredDeletesForDI_RenderThread(this->AsShared());
+	//check(SystemInstancesToProxyData.Contains(SystemInstance));
+	SystemInstancesToProxyData.Remove(SystemInstance);
 }
 
 //------------------------------------------------------------------------------------------------------------
@@ -562,7 +544,7 @@ UNiagaraDataInterfacePhysicsAsset::UNiagaraDataInterfacePhysicsAsset(FObjectInit
 	, SourceComponent(nullptr)
 	, PhysicsAsset(nullptr)
 {
-	Proxy = MakeShared<FNDIPhysicsAssetProxy, ESPMode::ThreadSafe>();
+	Proxy.Reset(new FNDIPhysicsAssetProxy());
 }
 
 void UNiagaraDataInterfacePhysicsAsset::ExtractSourceComponent(FNiagaraSystemInstance* SystemInstance)
diff --git a/Engine/Plugins/Experimental/HairStrands/Source/HairStrandsNiagara/Private/NiagaraDataInterfacePressureGrid.cpp b/Engine/Plugins/Experimental/HairStrands/Source/HairStrandsNiagara/Private/NiagaraDataInterfacePressureGrid.cpp
--- a/Engine/Plugins/Experimental/HairStrands/Source/HairStrandsNiagara/Private/NiagaraDataInterfacePressureGrid.cpp
+++ b/Engine/Plugins/Experimental/HairStrands/Source/HairStrandsNiagara/Private/NiagaraDataInterfacePressureGrid.cpp
@@ -259,7 +259,7 @@ UNiagaraDataInterfacePressureGrid::UNiagaraDataInterfacePressureGrid(FObjectInit
 	, MinOutput(0)
 	, MaxOutput(0)
 {
-	Proxy = MakeShared<FNDIPressureGridProxy, ESPMode::ThreadSafe>();
+	Proxy.Reset(new FNDIPressureGridProxy());
 }
 
 bool UNiagaraDataInterfacePressureGrid::InitPerInstanceData(void* PerInstanceData, FNiagaraSystemInstance* SystemInstance)
diff --git a/Engine/Plugins/Experimental/HairStrands/Source/HairStrandsNiagara/Public/NiagaraDataInterfaceHairStrands.h b/Engine/Plugins/Experimental/HairStrands/Source/HairStrandsNiagara/Public/NiagaraDataInterfaceHairStrands.h
--- a/Engine/Plugins/Experimental/HairStrands/Source/HairStrandsNiagara/Public/NiagaraDataInterfaceHairStrands.h
+++ b/Engine/Plugins/Experimental/HairStrands/Source/HairStrandsNiagara/Public/NiagaraDataInterfaceHairStrands.h
@@ -381,9 +381,6 @@ protected:
 /** Proxy to send data to gpu */
 struct FNDIHairStrandsProxy : public FNiagaraDataInterfaceProxy
 {
-	/** Destroy internal data */
-	virtual void DeferredDestroy() override;
-
 	/** Get the size of the data that will be passed to render*/
 	virtual int32 PerInstanceDataPassedToRenderThreadSize() const override { return sizeof(FNDIHairStrandsData); }
 
@@ -407,8 +404,5 @@ struct FNDIHairStrandsProxy : public FNiagaraDataInterfaceProxy
 
 	/** List of proxy data for each system instances*/
 	TMap<FNiagaraSystemInstanceID, FNDIHairStrandsData> SystemInstancesToProxyData;
-
-	/** List of proxy data to destroy later */
-	TSet<FNiagaraSystemInstanceID> DeferredDestroyList;
 };
 
diff --git a/Engine/Plugins/Experimental/HairStrands/Source/HairStrandsNiagara/Public/NiagaraDataInterfacePhysicsAsset.h b/Engine/Plugins/Experimental/HairStrands/Source/HairStrandsNiagara/Public/NiagaraDataInterfacePhysicsAsset.h
--- a/Engine/Plugins/Experimental/HairStrands/Source/HairStrandsNiagara/Public/NiagaraDataInterfacePhysicsAsset.h
+++ b/Engine/Plugins/Experimental/HairStrands/Source/HairStrandsNiagara/Public/NiagaraDataInterfacePhysicsAsset.h
@@ -212,9 +212,6 @@ protected:
 /** Proxy to send data to gpu */
 struct FNDIPhysicsAssetProxy : public FNiagaraDataInterfaceProxy
 {
-	/** Destroy internal data */
-	virtual void DeferredDestroy() override;
-
 	/** Get the size of the data that will be passed to render*/
 	virtual int32 PerInstanceDataPassedToRenderThreadSize() const override { return sizeof(FNDIPhysicsAssetData); }
 
@@ -229,8 +226,5 @@ struct FNDIPhysicsAssetProxy : public FNiagaraDataInterfaceProxy
 
 	/** List of proxy data for each system instances*/
 	TMap<FNiagaraSystemInstanceID, FNDIPhysicsAssetData> SystemInstancesToProxyData;
-
-	/** List of proxy data to destroy later */
-	TSet<FNiagaraSystemInstanceID> DeferredDestroyList;
 };
 
diff --git a/Engine/Plugins/Experimental/HairStrands/Source/HairStrandsNiagara/Public/NiagaraDataInterfacePressureGrid.h b/Engine/Plugins/Experimental/HairStrands/Source/HairStrandsNiagara/Public/NiagaraDataInterfacePressureGrid.h
--- a/Engine/Plugins/Experimental/HairStrands/Source/HairStrandsNiagara/Public/NiagaraDataInterfacePressureGrid.h
+++ b/Engine/Plugins/Experimental/HairStrands/Source/HairStrandsNiagara/Public/NiagaraDataInterfacePressureGrid.h
@@ -169,9 +169,6 @@ protected:
 /** Proxy to send data to gpu */
 struct FNDIPressureGridProxy : public FNiagaraDataInterfaceProxy
 {
-	/** Destroy internal data */
-	virtual void DeferredDestroy() override;
-
 	/** Get the size of the data that will be passed to render*/
 	virtual int32 PerInstanceDataPassedToRenderThreadSize() const override { return sizeof(FNDIPressureGridData); }
 
@@ -195,8 +192,5 @@ struct FNDIPressureGridProxy : public FNiagaraDataInterfaceProxy
 
 	/** List of proxy data for each system instances*/
 	TMap<FNiagaraSystemInstanceID, FNDIPressureGridData> SystemInstancesToProxyData;
-
-	/** List of proxy data to destroy later */
-	TSet<FNiagaraSystemInstanceID> DeferredDestroyList;
 };
 
diff --git a/Engine/Plugins/FX/HoudiniNiagara/Source/HoudiniNiagara/Private/NiagaraDataInterfaceHoudiniCSV.cpp b/Engine/Plugins/FX/HoudiniNiagara/Source/HoudiniNiagara/Private/NiagaraDataInterfaceHoudiniCSV.cpp
--- a/Engine/Plugins/FX/HoudiniNiagara/Source/HoudiniNiagara/Private/NiagaraDataInterfaceHoudiniCSV.cpp
+++ b/Engine/Plugins/FX/HoudiniNiagara/Source/HoudiniNiagara/Private/NiagaraDataInterfaceHoudiniCSV.cpp
@@ -96,7 +96,7 @@ UNiagaraDataInterfaceHoudiniCSV::UNiagaraDataInterfaceHoudiniCSV(FObjectInitiali
     HoudiniCSVAsset = nullptr;
 	LastSpawnedPointID = -1;
 
-	Proxy = MakeShared<FNiagaraDataInterfaceProxyHoudiniCSV, ESPMode::ThreadSafe>();
+	Proxy.Reset(new FNiagaraDataInterfaceProxyHoudiniCSV());
 }
 
 void UNiagaraDataInterfaceHoudiniCSV::PostInitProperties()
diff --git a/Engine/Plugins/FX/Niagara/Source/Niagara/Classes/NiagaraDataInterface.h b/Engine/Plugins/FX/Niagara/Source/Niagara/Classes/NiagaraDataInterface.h
--- a/Engine/Plugins/FX/Niagara/Source/Niagara/Classes/NiagaraDataInterface.h
+++ b/Engine/Plugins/FX/Niagara/Source/Niagara/Classes/NiagaraDataInterface.h
@@ -159,13 +159,12 @@ private:
 
 struct FNiagaraDataInterfaceProxy : TSharedFromThis<FNiagaraDataInterfaceProxy, ESPMode::ThreadSafe>
 {
+	FNiagaraDataInterfaceProxy() {}
 	virtual ~FNiagaraDataInterfaceProxy() {/*check(IsInRenderingThread());*/}
 
 	virtual int32 PerInstanceDataPassedToRenderThreadSize() const = 0;
 	virtual void ConsumePerInstanceDataFromGameThread(void* PerInstanceData, const FNiagaraSystemInstanceID& Instance) { check(false); }
 
-	virtual void DeferredDestroy() {}
-
 	// #todo(dmp): move all of this stuff to the RW interface to keep it out of here?
 
 	// a set of the shader stages that require the data interface for data output
@@ -316,7 +315,7 @@ protected:
 
 	virtual bool CopyToInternal(UNiagaraDataInterface* Destination) const;
 
-	TSharedPtr<FNiagaraDataInterfaceProxy, ESPMode::ThreadSafe> Proxy;
+	TUniquePtr<FNiagaraDataInterfaceProxy> Proxy;
 };
 
 /** Helper class for decoding NDI parameters into a usable struct type. */
diff --git a/Engine/Plugins/FX/Niagara/Source/Niagara/Classes/NiagaraDataInterfaceCurveBase.h b/Engine/Plugins/FX/Niagara/Source/Niagara/Classes/NiagaraDataInterfaceCurveBase.h
--- a/Engine/Plugins/FX/Niagara/Source/Niagara/Classes/NiagaraDataInterfaceCurveBase.h
+++ b/Engine/Plugins/FX/Niagara/Source/Niagara/Classes/NiagaraDataInterfaceCurveBase.h
@@ -103,7 +103,7 @@ public:
 		, ShowInCurveEditor(false)
 #endif
 	{
-		Proxy = MakeShared<FNiagaraDataInterfaceProxyCurveBase, ESPMode::ThreadSafe>();
+		Proxy.Reset(new FNiagaraDataInterfaceProxyCurveBase());
 	}
 
 	UNiagaraDataInterfaceCurveBase(FObjectInitializer const& ObjectInitializer)
@@ -118,7 +118,7 @@ public:
 		, ShowInCurveEditor(false)
 #endif
 	{
-		Proxy = MakeShared<FNiagaraDataInterfaceProxyCurveBase, ESPMode::ThreadSafe>();
+		Proxy.Reset(new FNiagaraDataInterfaceProxyCurveBase());
 	}
 
 	enum
diff --git a/Engine/Plugins/FX/Niagara/Source/Niagara/Classes/NiagaraDataInterfaceGrid2DCollection.h b/Engine/Plugins/FX/Niagara/Source/Niagara/Classes/NiagaraDataInterfaceGrid2DCollection.h
--- a/Engine/Plugins/FX/Niagara/Source/Niagara/Classes/NiagaraDataInterfaceGrid2DCollection.h
+++ b/Engine/Plugins/FX/Niagara/Source/Niagara/Classes/NiagaraDataInterfaceGrid2DCollection.h
@@ -91,16 +91,9 @@ struct FNiagaraDataInterfaceProxyGrid2DCollection : public FNiagaraDataInterface
 	virtual void PostStage(FRHICommandList& RHICmdList, const FNiagaraDataInterfaceSetArgs& Context) override;
 	virtual void ResetData(FRHICommandList& RHICmdList, const FNiagaraDataInterfaceSetArgs& Context) override;
 
-	virtual void DeferredDestroy() override;		
-	void DestroyPerInstanceData(NiagaraEmitterInstanceBatcher* Batcher, const FNiagaraSystemInstanceID& SystemInstance);
-
 	/* List of proxy data for each system instances*/
 	// #todo(dmp): this should all be refactored to avoid duplicate code
 	TMap<FNiagaraSystemInstanceID, Grid2DCollectionRWInstanceData> SystemInstancesToProxyData;
-
-	/* List of proxy data to destroy later */
-	// #todo(dmp): this should all be refactored to avoid duplicate code
-	TSet<FNiagaraSystemInstanceID> DeferredDestroyList;
 };
 
 UCLASS(EditInlineNew, Category = "Grid", meta = (DisplayName = "Grid2D Collection"), Blueprintable, BlueprintType)
diff --git a/Engine/Plugins/FX/Niagara/Source/Niagara/Classes/NiagaraDataInterfaceNeighborGrid3D.h b/Engine/Plugins/FX/Niagara/Source/Niagara/Classes/NiagaraDataInterfaceNeighborGrid3D.h
--- a/Engine/Plugins/FX/Niagara/Source/Niagara/Classes/NiagaraDataInterfaceNeighborGrid3D.h
+++ b/Engine/Plugins/FX/Niagara/Source/Niagara/Classes/NiagaraDataInterfaceNeighborGrid3D.h
@@ -34,16 +34,10 @@ struct FNiagaraDataInterfaceProxyNeighborGrid3D : public FNiagaraDataInterfacePr
 	virtual void PreStage(FRHICommandList& RHICmdList, const FNiagaraDataInterfaceSetArgs& Context) override;	
 	virtual void ConsumePerInstanceDataFromGameThread(void* PerInstanceData, const FNiagaraSystemInstanceID& Instance) override {}
 	virtual int32 PerInstanceDataPassedToRenderThreadSize() const override { return sizeof(NeighborGrid3DRWInstanceData); }	
-	void DestroyPerInstanceData(NiagaraEmitterInstanceBatcher* Batcher, const FNiagaraSystemInstanceID& SystemInstance);
-	void DeferredDestroy();
 
 	/* List of proxy data for each system instances*/
 	// #todo(dmp): this should all be refactored to avoid duplicate code
 	TMap<FNiagaraSystemInstanceID, NeighborGrid3DRWInstanceData> SystemInstancesToProxyData;
-
-	/* List of proxy data to destroy later */
-	// #todo(dmp): this should all be refactored to avoid duplicate code
-	TSet<FNiagaraSystemInstanceID> DeferredDestroyList;
 };
 
 UCLASS(EditInlineNew, Category = "Grid", meta = (DisplayName = "Neighbor Grid3D"))
diff --git a/Engine/Plugins/FX/Niagara/Source/Niagara/Classes/NiagaraDataInterfaceStaticMesh.h b/Engine/Plugins/FX/Niagara/Source/Niagara/Classes/NiagaraDataInterfaceStaticMesh.h
--- a/Engine/Plugins/FX/Niagara/Source/Niagara/Classes/NiagaraDataInterfaceStaticMesh.h
+++ b/Engine/Plugins/FX/Niagara/Source/Niagara/Classes/NiagaraDataInterfaceStaticMesh.h
@@ -354,8 +354,6 @@ struct FNiagaraPassedInstanceDataForRT
 
 struct FNiagaraDataInterfaceProxyStaticMesh : public FNiagaraDataInterfaceProxy
 {
-	virtual void DeferredDestroy() override;
-
 	virtual int32 PerInstanceDataPassedToRenderThreadSize() const override
 	{
 		return sizeof(FNiagaraPassedInstanceDataForRT);
@@ -367,6 +365,4 @@ struct FNiagaraDataInterfaceProxyStaticMesh : public FNiagaraDataInterfaceProxy
 	void DestroyPerInstanceData(NiagaraEmitterInstanceBatcher* Batcher, const FNiagaraSystemInstanceID& SystemInstance);
 
 	TMap<FNiagaraSystemInstanceID, FNiagaraStaticMeshData> SystemInstancesToMeshData;
-
-	TSet<FNiagaraSystemInstanceID> DeferredDestroyList;
 };
diff --git a/Engine/Plugins/FX/Niagara/Source/Niagara/Classes/NiagaraEmitterInstanceBatcher.h b/Engine/Plugins/FX/Niagara/Source/Niagara/Classes/NiagaraEmitterInstanceBatcher.h
--- a/Engine/Plugins/FX/Niagara/Source/Niagara/Classes/NiagaraEmitterInstanceBatcher.h
+++ b/Engine/Plugins/FX/Niagara/Source/Niagara/Classes/NiagaraEmitterInstanceBatcher.h
@@ -76,8 +76,6 @@ public:
 	/** Called to release GPU instance counts that the batcher is tracking. */
 	void ReleaseInstanceCounts_RenderThread(FNiagaraComputeExecutionContext* ExecContext, FNiagaraDataSet* DataSet);
 
-	NIAGARA_API void EnqueueDeferredDeletesForDI_RenderThread(TSharedPtr<FNiagaraDataInterfaceProxy, ESPMode::ThreadSafe> Proxy);
-
 #if WITH_EDITOR
 	virtual void Suspend() override {}
 	virtual void Resume() override {}
@@ -212,12 +210,6 @@ private:
 	TArray<FNiagaraGPUSystemTick> Ticks_RT;
 	FGlobalDistanceFieldParameterData GlobalDistanceFieldParams;
 
-	// For deferred deletion of resources
-	TSet<FNiagaraComputeExecutionContext*> ContextsToDestroy_RT;
-	TSet<FNiagaraDataSet*> DataSetsToDestroy_RT;
-
-	TSet<TSharedPtr<FNiagaraDataInterfaceProxy, ESPMode::ThreadSafe>> DIProxyDeferredDeletes_RT;
-
 	/** A buffer used by the compute shader which determines the list of free particle IDs for each emitter. */
 	FRWBuffer FreeIDListSizesBuffer;
 	uint32 NumAllocatedFreeIDListSizes;
diff --git a/Engine/Plugins/FX/Niagara/Source/Niagara/Private/NiagaraDataInterfaceCamera.cpp b/Engine/Plugins/FX/Niagara/Source/Niagara/Private/NiagaraDataInterfaceCamera.cpp
--- a/Engine/Plugins/FX/Niagara/Source/Niagara/Private/NiagaraDataInterfaceCamera.cpp
+++ b/Engine/Plugins/FX/Niagara/Source/Niagara/Private/NiagaraDataInterfaceCamera.cpp
@@ -21,7 +21,7 @@ const FName UNiagaraDataInterfaceCamera::GetFieldOfViewName(TEXT("GetFieldOfView
 UNiagaraDataInterfaceCamera::UNiagaraDataInterfaceCamera(FObjectInitializer const& ObjectInitializer)
 	: Super(ObjectInitializer)
 {
-	Proxy = MakeShared<FNiagaraDataIntefaceProxyCameraQuery, ESPMode::ThreadSafe>();
+	Proxy.Reset(new FNiagaraDataIntefaceProxyCameraQuery());
 }
 
 void UNiagaraDataInterfaceCamera::PostInitProperties()
diff --git a/Engine/Plugins/FX/Niagara/Source/Niagara/Private/NiagaraDataInterfaceCollisionQuery.cpp b/Engine/Plugins/FX/Niagara/Source/Niagara/Private/NiagaraDataInterfaceCollisionQuery.cpp
--- a/Engine/Plugins/FX/Niagara/Source/Niagara/Private/NiagaraDataInterfaceCollisionQuery.cpp
+++ b/Engine/Plugins/FX/Niagara/Source/Niagara/Private/NiagaraDataInterfaceCollisionQuery.cpp
@@ -19,7 +19,7 @@ UNiagaraDataInterfaceCollisionQuery::UNiagaraDataInterfaceCollisionQuery(FObject
 {
 	TraceChannelEnum = StaticEnum<ECollisionChannel>();
 
-    Proxy = MakeShared<FNiagaraDataIntefaceProxyCollisionQuery, ESPMode::ThreadSafe>();
+    Proxy.Reset(new FNiagaraDataIntefaceProxyCollisionQuery());
 }
 
 bool UNiagaraDataInterfaceCollisionQuery::InitPerInstanceData(void* PerInstanceData, FNiagaraSystemInstance* InSystemInstance)
diff --git a/Engine/Plugins/FX/Niagara/Source/Niagara/Private/NiagaraDataInterfaceCurlNoise.cpp b/Engine/Plugins/FX/Niagara/Source/Niagara/Private/NiagaraDataInterfaceCurlNoise.cpp
--- a/Engine/Plugins/FX/Niagara/Source/Niagara/Private/NiagaraDataInterfaceCurlNoise.cpp
+++ b/Engine/Plugins/FX/Niagara/Source/Niagara/Private/NiagaraDataInterfaceCurlNoise.cpp
@@ -309,7 +309,7 @@ UNiagaraDataInterfaceCurlNoise::UNiagaraDataInterfaceCurlNoise(FObjectInitialize
 {
 	OffsetFromSeed = FNiagaraUIntVectorToFVector(Rand3DPCG16(FIntVector(Seed, Seed, Seed))) / 100.0;
 
-	Proxy = MakeShared<FNiagaraDataInterfaceProxyCurlNoise, ESPMode::ThreadSafe>(OffsetFromSeed);
+	Proxy.Reset(new FNiagaraDataInterfaceProxyCurlNoise(OffsetFromSeed));
 }
 
 void UNiagaraDataInterfaceCurlNoise::PostInitProperties()
diff --git a/Engine/Plugins/FX/Niagara/Source/Niagara/Private/NiagaraDataInterfaceGrid2DCollection.cpp b/Engine/Plugins/FX/Niagara/Source/Niagara/Private/NiagaraDataInterfaceGrid2DCollection.cpp
--- a/Engine/Plugins/FX/Niagara/Source/Niagara/Private/NiagaraDataInterfaceGrid2DCollection.cpp
+++ b/Engine/Plugins/FX/Niagara/Source/Niagara/Private/NiagaraDataInterfaceGrid2DCollection.cpp
@@ -124,7 +124,7 @@ private:
 UNiagaraDataInterfaceGrid2DCollection::UNiagaraDataInterfaceGrid2DCollection(FObjectInitializer const& ObjectInitializer)
 	: Super(ObjectInitializer)
 {
-	Proxy = MakeShared<FNiagaraDataInterfaceProxyGrid2DCollection, ESPMode::ThreadSafe>();
+	Proxy.Reset(new FNiagaraDataInterfaceProxyGrid2DCollection());
 	RWProxy = (FNiagaraDataInterfaceProxyRW*) Proxy.Get();
 	PushToRenderThread();
 }
@@ -399,15 +399,8 @@ bool UNiagaraDataInterfaceGrid2DCollection::InitPerInstanceData(void* PerInstanc
 	ENQUEUE_RENDER_COMMAND(FUpdateData)(
 		[RT_Proxy, RT_NumCellsX, RT_NumCellsY, RT_NumTilesX, RT_NumTilesY, RT_CellSize, RT_WorldBBoxSize, RT_OutputShaderStages, RT_IterationShaderStages, InstanceID = SystemInstance->GetId()](FRHICommandListImmediate& RHICmdList)
 	{
-		Grid2DCollectionRWInstanceData* TargetData = RT_Proxy->SystemInstancesToProxyData.Find(InstanceID);
-		if (TargetData != nullptr)
-		{
-			RT_Proxy->DeferredDestroyList.Remove(InstanceID);
-		}
-		else
-		{
-			TargetData = &RT_Proxy->SystemInstancesToProxyData.Add(InstanceID);
-		}
+		check(!RT_Proxy->SystemInstancesToProxyData.Contains(InstanceID));
+		Grid2DCollectionRWInstanceData* TargetData = &RT_Proxy->SystemInstancesToProxyData.Add(InstanceID);
 
 		TargetData->NumCellsX = RT_NumCellsX;
 		TargetData->NumCellsY = RT_NumCellsY;
@@ -440,7 +433,8 @@ void UNiagaraDataInterfaceGrid2DCollection::DestroyPerInstanceData(void* PerInst
 	ENQUEUE_RENDER_COMMAND(FNiagaraDIDestroyInstanceData) (
 		[ThisProxy, InstanceID = SystemInstance->GetId(), Batcher = SystemInstance->GetBatcher()](FRHICommandListImmediate& CmdList)
 	{
-		ThisProxy->DestroyPerInstanceData(Batcher, InstanceID);
+		//check(ThisProxy->SystemInstancesToProxyData.Contains(InstanceID));
+		ThisProxy->SystemInstancesToProxyData.Remove(InstanceID);
 	}
 	);
 }
@@ -661,27 +655,6 @@ void UNiagaraDataInterfaceGrid2DCollection::GetCellSize(FVectorVMContext& Contex
 	}
 }
 
-// #todo(dmp): move these to super class
-void FNiagaraDataInterfaceProxyGrid2DCollection::DestroyPerInstanceData(NiagaraEmitterInstanceBatcher* Batcher, const FNiagaraSystemInstanceID& SystemInstance)
-{
-	check(IsInRenderingThread());
-
-	DeferredDestroyList.Add(SystemInstance);
-	Batcher->EnqueueDeferredDeletesForDI_RenderThread(this->AsShared());
-}
-
-// #todo(dmp): move these to super class
-void FNiagaraDataInterfaceProxyGrid2DCollection::DeferredDestroy()
-{
-	for (const FNiagaraSystemInstanceID& Sys : DeferredDestroyList)
-	{
-		SystemInstancesToProxyData.Remove(Sys);
-	}
-
-	DeferredDestroyList.Empty();
-}
-
-
 void FNiagaraDataInterfaceProxyGrid2DCollection::PreStage(FRHICommandList& RHICmdList, const FNiagaraDataInterfaceSetArgs& Context)
 {
 	// #todo(dmp): Context doesnt need to specify if a stage is output or not since we moved pre/post stage to the DI itself.  Not sure which design is better for the future
diff --git a/Engine/Plugins/FX/Niagara/Source/Niagara/Private/NiagaraDataInterfaceNeighborGrid3D.cpp b/Engine/Plugins/FX/Niagara/Source/Niagara/Private/NiagaraDataInterfaceNeighborGrid3D.cpp
--- a/Engine/Plugins/FX/Niagara/Source/Niagara/Private/NiagaraDataInterfaceNeighborGrid3D.cpp
+++ b/Engine/Plugins/FX/Niagara/Source/Niagara/Private/NiagaraDataInterfaceNeighborGrid3D.cpp
@@ -140,7 +140,7 @@ UNiagaraDataInterfaceNeighborGrid3D::UNiagaraDataInterfaceNeighborGrid3D(FObject
 	: Super(ObjectInitializer)
 	, MaxNeighborsPerVoxel(8)
 {
-	Proxy = MakeShared<FNiagaraDataInterfaceProxyNeighborGrid3D, ESPMode::ThreadSafe>();
+	Proxy.Reset(new FNiagaraDataInterfaceProxyNeighborGrid3D());
 	PushToRenderThread();
 }
 
@@ -438,15 +438,8 @@ bool UNiagaraDataInterfaceNeighborGrid3D::InitPerInstanceData(void* PerInstanceD
 	ENQUEUE_RENDER_COMMAND(FUpdateData)(
 		[RT_Proxy, RT_NumVoxels, RT_MaxNeighborsPerVoxel, RT_WorldBBoxSize, RT_OutputShaderStages, RT_IterationShaderStages, InstanceID = SystemInstance->GetId()](FRHICommandListImmediate& RHICmdList)
 	{
-		NeighborGrid3DRWInstanceData* TargetData = RT_Proxy->SystemInstancesToProxyData.Find(InstanceID);
-		if (TargetData != nullptr)
-		{
-			RT_Proxy->DeferredDestroyList.Remove(InstanceID);
-		}
-		else
-		{
-			TargetData = &RT_Proxy->SystemInstancesToProxyData.Add(InstanceID);
-		}
+		check(!RT_Proxy->SystemInstancesToProxyData.Contains(InstanceID));
+		NeighborGrid3DRWInstanceData* TargetData = &RT_Proxy->SystemInstancesToProxyData.Add(InstanceID);
 
 		TargetData->NumVoxels = RT_NumVoxels;
 		TargetData->MaxNeighborsPerVoxel = RT_MaxNeighborsPerVoxel;		
@@ -476,7 +469,8 @@ void UNiagaraDataInterfaceNeighborGrid3D::DestroyPerInstanceData(void* PerInstan
 	ENQUEUE_RENDER_COMMAND(FNiagaraDIDestroyInstanceData) (
 		[ThisProxy, InstanceID = SystemInstance->GetId(), Batcher = SystemInstance->GetBatcher()](FRHICommandListImmediate& CmdList)
 	{
-		ThisProxy->DestroyPerInstanceData(Batcher, InstanceID);
+		//check(ThisProxy->SystemInstancesToProxyData.Contains(InstanceID));
+		ThisProxy->SystemInstancesToProxyData.Remove(InstanceID);
 	}
 	);
 }
@@ -492,26 +486,6 @@ void FNiagaraDataInterfaceProxyNeighborGrid3D::PreStage(FRHICommandList& RHICmdL
 	}
 }
 
-// #todo(dmp): move these to super class
-void FNiagaraDataInterfaceProxyNeighborGrid3D::DestroyPerInstanceData(NiagaraEmitterInstanceBatcher* Batcher, const FNiagaraSystemInstanceID& SystemInstance)
-{
-	check(IsInRenderingThread());
-
-	DeferredDestroyList.Add(SystemInstance);
-	Batcher->EnqueueDeferredDeletesForDI_RenderThread(this->AsShared());
-}
-
-// #todo(dmp): move these to super class
-void FNiagaraDataInterfaceProxyNeighborGrid3D::DeferredDestroy()
-{
-	for (const FNiagaraSystemInstanceID& Sys : DeferredDestroyList)
-	{
-		SystemInstancesToProxyData.Remove(Sys);
-	}
-
-	DeferredDestroyList.Empty();
-}
-
 bool UNiagaraDataInterfaceNeighborGrid3D::CopyToInternal(UNiagaraDataInterface* Destination) const
 {
 	if (!Super::CopyToInternal(Destination))
diff --git a/Engine/Plugins/FX/Niagara/Source/Niagara/Private/NiagaraDataInterfaceOcclusion.cpp b/Engine/Plugins/FX/Niagara/Source/Niagara/Private/NiagaraDataInterfaceOcclusion.cpp
--- a/Engine/Plugins/FX/Niagara/Source/Niagara/Private/NiagaraDataInterfaceOcclusion.cpp
+++ b/Engine/Plugins/FX/Niagara/Source/Niagara/Private/NiagaraDataInterfaceOcclusion.cpp
@@ -13,7 +13,7 @@ const FName UNiagaraDataInterfaceOcclusion::GetCameraOcclusionCircleName(TEXT("Q
 UNiagaraDataInterfaceOcclusion::UNiagaraDataInterfaceOcclusion(FObjectInitializer const& ObjectInitializer)
 	: Super(ObjectInitializer)
 {
-	Proxy = MakeShared<FNiagaraDataIntefaceProxyOcclusionQuery, ESPMode::ThreadSafe>();
+	Proxy.Reset(new FNiagaraDataIntefaceProxyOcclusionQuery());
 }
 
 void UNiagaraDataInterfaceOcclusion::PostInitProperties()
diff --git a/Engine/Plugins/FX/Niagara/Source/Niagara/Private/NiagaraDataInterfaceParticleRead.cpp b/Engine/Plugins/FX/Niagara/Source/Niagara/Private/NiagaraDataInterfaceParticleRead.cpp
--- a/Engine/Plugins/FX/Niagara/Source/Niagara/Private/NiagaraDataInterfaceParticleRead.cpp
+++ b/Engine/Plugins/FX/Niagara/Source/Niagara/Private/NiagaraDataInterfaceParticleRead.cpp
@@ -356,7 +356,7 @@ private:
 UNiagaraDataInterfaceParticleRead::UNiagaraDataInterfaceParticleRead(FObjectInitializer const& ObjectInitializer)
 	: Super(ObjectInitializer)
 {
-	Proxy = MakeShared<FNiagaraDataInterfaceProxyParticleRead, ESPMode::ThreadSafe>();
+	Proxy.Reset(new FNiagaraDataInterfaceProxyParticleRead());
 }
 
 void UNiagaraDataInterfaceParticleRead::PostInitProperties()
diff --git a/Engine/Plugins/FX/Niagara/Source/Niagara/Private/NiagaraDataInterfaceRW.cpp b/Engine/Plugins/FX/Niagara/Source/Niagara/Private/NiagaraDataInterfaceRW.cpp
--- a/Engine/Plugins/FX/Niagara/Source/Niagara/Private/NiagaraDataInterfaceRW.cpp
+++ b/Engine/Plugins/FX/Niagara/Source/Niagara/Private/NiagaraDataInterfaceRW.cpp
@@ -50,7 +50,7 @@ UNiagaraDataInterfaceGrid3D::UNiagaraDataInterfaceGrid3D(FObjectInitializer cons
 	, SetGridFromVoxelSize(false)	
 	, WorldBBoxSize(100., 100., 100.)
 {
-	Proxy = MakeShared<FNiagaraDataInterfaceProxyRW, ESPMode::ThreadSafe>();
+	Proxy.Reset(new FNiagaraDataInterfaceProxyRW());
 	RWProxy = (FNiagaraDataInterfaceProxyRW*) Proxy.Get();
 	PushToRenderThread();
 }
@@ -319,7 +319,7 @@ UNiagaraDataInterfaceGrid2D::UNiagaraDataInterfaceGrid2D(FObjectInitializer cons
 	, SetGridFromMaxAxis(false)	
 	, WorldBBoxSize(100., 100.)
 {
-	Proxy = MakeShared<FNiagaraDataInterfaceProxyRW, ESPMode::ThreadSafe>();
+	Proxy.Reset(new FNiagaraDataInterfaceProxyRW());
 	PushToRenderThread();
 }
 
diff --git a/Engine/Plugins/FX/Niagara/Source/Niagara/Private/NiagaraDataInterfaceSkeletalMesh.cpp b/Engine/Plugins/FX/Niagara/Source/Niagara/Private/NiagaraDataInterfaceSkeletalMesh.cpp
--- a/Engine/Plugins/FX/Niagara/Source/Niagara/Private/NiagaraDataInterfaceSkeletalMesh.cpp
+++ b/Engine/Plugins/FX/Niagara/Source/Niagara/Private/NiagaraDataInterfaceSkeletalMesh.cpp
@@ -1575,7 +1575,7 @@ UNiagaraDataInterfaceSkeletalMesh::UNiagaraDataInterfaceSkeletalMesh(FObjectInit
 #endif
 	, ChangeId(0)
 {
-	Proxy = MakeShared<FNiagaraDataInterfaceProxySkeletalMesh, ESPMode::ThreadSafe>();
+	Proxy.Reset(new FNiagaraDataInterfaceProxySkeletalMesh());
 }
 
 
@@ -1741,6 +1741,7 @@ void UNiagaraDataInterfaceSkeletalMesh::DestroyPerInstanceData(void* PerInstance
 		ENQUEUE_RENDER_COMMAND(FNiagaraDestroySkeletalMeshInstanceData) (
 			[ThisProxy, InstanceID = SystemInstance->GetId()](FRHICommandListImmediate& CmdList)
 			{
+				//check(ThisProxy->Contains(InstanceID));
 				ThisProxy->SystemInstancesToData.Remove(InstanceID);
 			}
 		);
diff --git a/Engine/Plugins/FX/Niagara/Source/Niagara/Private/NiagaraDataInterfaceStaticMesh.cpp b/Engine/Plugins/FX/Niagara/Source/Niagara/Private/NiagaraDataInterfaceStaticMesh.cpp
--- a/Engine/Plugins/FX/Niagara/Source/Niagara/Private/NiagaraDataInterfaceStaticMesh.cpp
+++ b/Engine/Plugins/FX/Niagara/Source/Niagara/Private/NiagaraDataInterfaceStaticMesh.cpp
@@ -574,53 +574,20 @@ private:
 
 //////////////////////////////////////////////////////////////////////////
 
-void FNiagaraDataInterfaceProxyStaticMesh::DeferredDestroy()
-{
-	//-TODO: This is incorrect, we could be destroying instance data for a batcher is yet to tick, we should only be destroying data for this batcher!
-	for (const FNiagaraSystemInstanceID& Sys : DeferredDestroyList)
-	{
-		SystemInstancesToMeshData.Remove(Sys);
-		//UE_LOG(LogNiagara, Log, TEXT("RT: StaticMesh DI - DeferredDestroy %s"), *Sys.ToString());
-	}
-
-	DeferredDestroyList.Empty();
-}
-
 void FNiagaraDataInterfaceProxyStaticMesh::InitializePerInstanceData(const FNiagaraSystemInstanceID& SystemInstance, FStaticMeshGpuSpawnBuffer* MeshGPUSpawnBuffer)
 {
 	check(IsInRenderingThread());
+	check(!SystemInstancesToMeshData.Contains(SystemInstance));
 
-	FNiagaraStaticMeshData* Data = SystemInstancesToMeshData.Find(SystemInstance);
-	if (Data != nullptr)
-	{
-		DeferredDestroyList.Remove(SystemInstance);
-	}
-	else
-	{
-		Data = &SystemInstancesToMeshData.Add(SystemInstance);
-	}
-	//UE_LOG(LogNiagara, Log, TEXT("RT: StaticMesh DI - InitializePerInstanceData %s"), *SystemInstance.ToString());
-
-	// @todo-threadsafety We should not ever see this case! Though it's not really an error...
-	if (Data->MeshGpuSpawnBuffer)
-	{
-		Data->MeshGpuSpawnBuffer->ReleaseResource();
-		delete Data->MeshGpuSpawnBuffer;
-	}
-
-	Data->MeshGpuSpawnBuffer = MeshGPUSpawnBuffer;
+	FNiagaraStaticMeshData& Data = SystemInstancesToMeshData.Add(SystemInstance);
+	Data.MeshGpuSpawnBuffer = MeshGPUSpawnBuffer;
 }
 
 void FNiagaraDataInterfaceProxyStaticMesh::DestroyPerInstanceData(NiagaraEmitterInstanceBatcher* Batcher, const FNiagaraSystemInstanceID& SystemInstance)
 {
 	check(IsInRenderingThread());
-
-	//UE_LOG(LogNiagara, Log, TEXT("RT: StaticMesh DI - DestroyPerInstanceData %s"), *SystemInstance.ToString());
-
-	// @todo-threadsafety verify this destroys the MeshGPUSpawnBuffer data. This thread owns it now.
-	//SystemInstancesToMeshData.Remove(SystemInstance);
-	DeferredDestroyList.Add(SystemInstance);
-	Batcher->EnqueueDeferredDeletesForDI_RenderThread(this->AsShared());
+	//check(SystemInstancesToMeshData.Contains(SystemInstance));
+	SystemInstancesToMeshData.Remove(SystemInstance);
 }
 
 void FNiagaraDataInterfaceProxyStaticMesh::ConsumePerInstanceDataFromGameThread(void* PerInstanceData, const FNiagaraSystemInstanceID& Instance)
@@ -655,7 +622,7 @@ UNiagaraDataInterfaceStaticMesh::UNiagaraDataInterfaceStaticMesh(FObjectInitiali
 	//, bSupportingVertexColorSampling(0)//Vertex color filtering needs some more work.
 	//, bFilterInitialized(false)
 {
-	Proxy = MakeShared<FNiagaraDataInterfaceProxyStaticMesh, ESPMode::ThreadSafe>();
+	Proxy.Reset(new FNiagaraDataInterfaceProxyStaticMesh());
 }
 
 void UNiagaraDataInterfaceStaticMesh::PostInitProperties()
@@ -2488,4 +2455,3 @@ void FNDI_StaticMesh_GeneratedData::CleanupDynamicColorFilterData()
 }
 
 #undef LOCTEXT_NAMESPACE
-
diff --git a/Engine/Plugins/FX/Niagara/Source/Niagara/Private/NiagaraDataInterfaceTexture.cpp b/Engine/Plugins/FX/Niagara/Source/Niagara/Private/NiagaraDataInterfaceTexture.cpp
--- a/Engine/Plugins/FX/Niagara/Source/Niagara/Private/NiagaraDataInterfaceTexture.cpp
+++ b/Engine/Plugins/FX/Niagara/Source/Niagara/Private/NiagaraDataInterfaceTexture.cpp
@@ -20,7 +20,7 @@ UNiagaraDataInterfaceTexture::UNiagaraDataInterfaceTexture(FObjectInitializer co
 	: Super(ObjectInitializer)
 	, Texture(nullptr)
 {
-	Proxy = MakeShared<FNiagaraDataInterfaceProxyTexture, ESPMode::ThreadSafe>();
+	Proxy.Reset(new FNiagaraDataInterfaceProxyTexture());
 	PushToRenderThread();
 }
 
diff --git a/Engine/Plugins/FX/Niagara/Source/Niagara/Private/NiagaraDataInterfaceVectorField.cpp b/Engine/Plugins/FX/Niagara/Source/Niagara/Private/NiagaraDataInterfaceVectorField.cpp
--- a/Engine/Plugins/FX/Niagara/Source/Niagara/Private/NiagaraDataInterfaceVectorField.cpp
+++ b/Engine/Plugins/FX/Niagara/Source/Niagara/Private/NiagaraDataInterfaceVectorField.cpp
@@ -32,7 +32,7 @@ UNiagaraDataInterfaceVectorField::UNiagaraDataInterfaceVectorField(FObjectInitia
 	, bTileY(false)
 	, bTileZ(false)
 {
-	Proxy = MakeShared<FNiagaraDataInterfaceProxyVectorField, ESPMode::ThreadSafe>();
+	Proxy.Reset(new FNiagaraDataInterfaceProxyVectorField());
 	PushToRenderThread();
 }
 
diff --git a/Engine/Plugins/FX/Niagara/Source/Niagara/Private/NiagaraEmitterInstanceBatcher.cpp b/Engine/Plugins/FX/Niagara/Source/Niagara/Private/NiagaraEmitterInstanceBatcher.cpp
--- a/Engine/Plugins/FX/Niagara/Source/Niagara/Private/NiagaraEmitterInstanceBatcher.cpp
+++ b/Engine/Plugins/FX/Niagara/Source/Niagara/Private/NiagaraEmitterInstanceBatcher.cpp
@@ -157,39 +157,9 @@ void NiagaraEmitterInstanceBatcher::ReleaseInstanceCounts_RenderThread(FNiagaraC
 	}
 }
 
-void NiagaraEmitterInstanceBatcher::EnqueueDeferredDeletesForDI_RenderThread(TSharedPtr<FNiagaraDataInterfaceProxy, ESPMode::ThreadSafe> Proxy)
-{
-	DIProxyDeferredDeletes_RT.Add(MoveTemp(Proxy));
-}
-
 void NiagaraEmitterInstanceBatcher::FinishDispatches()
 {
 	ReleaseTicks();
-
-	for (FNiagaraComputeExecutionContext* Context : ContextsToDestroy_RT)
-	{
-		check(Context);
-		// Put back the GPU instance counter the global pool.
-		GPUInstanceCounterManager.FreeEntry(Context->EmitterInstanceReadback.GPUCountOffset);
-		delete Context;
-	}
-	ContextsToDestroy_RT.Reset();
-
-	for (FNiagaraDataSet* DataSet : DataSetsToDestroy_RT)
-	{
-		check(DataSet);
-		// Put back the GPU instance counter the global pool.
-		DataSet->ReleaseGPUInstanceCounts(GPUInstanceCounterManager);
-		delete DataSet;
-	}
-	DataSetsToDestroy_RT.Reset();
-
-	for (TSharedPtr<FNiagaraDataInterfaceProxy, ESPMode::ThreadSafe>& Proxy : DIProxyDeferredDeletes_RT)
-	{
-		Proxy->DeferredDestroy();
-	}
-
-	DIProxyDeferredDeletes_RT.Empty();
 }
 
 void NiagaraEmitterInstanceBatcher::ReleaseTicks()
diff --git a/Engine/Plugins/FX/Niagara/Source/Niagara/Private/NiagaraScriptExecutionContext.cpp b/Engine/Plugins/FX/Niagara/Source/Niagara/Private/NiagaraScriptExecutionContext.cpp
--- a/Engine/Plugins/FX/Niagara/Source/Niagara/Private/NiagaraScriptExecutionContext.cpp
+++ b/Engine/Plugins/FX/Niagara/Source/Niagara/Private/NiagaraScriptExecutionContext.cpp
@@ -240,11 +240,6 @@ bool FNiagaraScriptExecutionContext::CanExecute()const
 	return Script && Script->GetVMExecutableData().IsValid() && Script->GetVMExecutableData().ByteCode.Num() > 0;
 }
 
-struct H2
-{
-	FNiagaraDataInterfaceProxy* Proxy;
-};
-
 void FNiagaraGPUSystemTick::Init(FNiagaraSystemInstance* InSystemInstance)
 {
 	check(IsInGameThread());
diff --git a/Engine/Plugins/FX/Niagara/Source/Niagara/Private/NiagaraSystemInstance.cpp b/Engine/Plugins/FX/Niagara/Source/Niagara/Private/NiagaraSystemInstance.cpp
--- a/Engine/Plugins/FX/Niagara/Source/Niagara/Private/NiagaraSystemInstance.cpp
+++ b/Engine/Plugins/FX/Niagara/Source/Niagara/Private/NiagaraSystemInstance.cpp
@@ -1148,6 +1148,7 @@ void FNiagaraSystemInstance::InitDataInterfaces()
 
 	Component->GetOverrideParameters().Tick();
 	
+	//-TODO: Validate that any queued ticks have been executed
 	DestroyDataInterfaceInstanceData();
 
 	GPUDataInterfaceInstanceDataSize = 0;
