Chaos: progress towards bringing NiagaraDataInterface up to date.

diff --git a/Engine/Plugins/Experimental/ChaosNiagara/Source/ChaosNiagara/Classes/NiagaraDataInterfaceChaosDestruction.h b/Engine/Plugins/Experimental/ChaosNiagara/Source/ChaosNiagara/Classes/NiagaraDataInterfaceChaosDestruction.h
--- a/Engine/Plugins/Experimental/ChaosNiagara/Source/ChaosNiagara/Classes/NiagaraDataInterfaceChaosDestruction.h
+++ b/Engine/Plugins/Experimental/ChaosNiagara/Source/ChaosNiagara/Classes/NiagaraDataInterfaceChaosDestruction.h
@@ -86,7 +86,7 @@ struct FSolverData
 		: Solver(nullptr)
 	{}
 
-	TSharedPtr<FPhysScene_Chaos> PhysScene;
+	FPhysScene* PhysScene;
 	Chaos::FPhysicsSolver* Solver;
 };
 
@@ -679,6 +679,9 @@ public:
 
 	virtual void ProvidePerInstanceDataForRenderThread(void* DataForRenderThread, void* PerInstanceData, const FNiagaraSystemInstanceID& SystemInstance) override;
 protected:
+#if INCLUDE_CHAOS
+	void RegisterWithSolverEventManager(Chaos::FPhysicsSolver* Solver);
+#endif
 	virtual bool CopyToInternal(UNiagaraDataInterface* Destination) const override;
 
 	void ResetInstData(FNDIChaosDestruction_InstanceData* InstData);
diff --git a/Engine/Plugins/Experimental/ChaosNiagara/Source/ChaosNiagara/Private/NiagaraDataInterfaceChaosDestruction.cpp b/Engine/Plugins/Experimental/ChaosNiagara/Source/ChaosNiagara/Private/NiagaraDataInterfaceChaosDestruction.cpp
--- a/Engine/Plugins/Experimental/ChaosNiagara/Source/ChaosNiagara/Private/NiagaraDataInterfaceChaosDestruction.cpp
+++ b/Engine/Plugins/Experimental/ChaosNiagara/Source/ChaosNiagara/Private/NiagaraDataInterfaceChaosDestruction.cpp
@@ -111,7 +111,6 @@ UNiagaraDataInterfaceChaosDestruction::UNiagaraDataInterfaceChaosDestruction(FOb
 	, FinalVelocityMagnitudeMinMax(FVector2D(-1.f, -1.f))
 	, MaxLatency(1.f)
 	, DebugType(EDebugTypeEnum::ChaosNiagara_DebugType_NoDebug)
-//	, ParticleIndexToProcess(-1)
 	, LastSpawnedPointID(-1)
 	, LastSpawnTime(-1.f)
 	, SolverTime(0.f)
@@ -170,20 +169,6 @@ void UNiagaraDataInterfaceChaosDestruction::PostLoad()
 	LastSpawnTime = -1.f;
 	TimeStampOfLastProcessedData = -1.f;
 
-#if WITH_CHAOS
-	if (GetWorld() && GetWorld()->GetPhysicsScene())
-	{
-		FPhysScene* Scene = GetWorld()->GetPhysicsScene();
-		Chaos::FEventManager* EventManager = Scene->GetScene().GetSolver()->GetEventManager();
-		if (EventManager)
-		{
-			EventManager->RegisterHandler<Chaos::FCollisionEventData>(Chaos::EEventType::Collision, this, &UNiagaraDataInterfaceChaosDestruction::HandleCollisionEvents);
-			EventManager->RegisterHandler<Chaos::FBreakingEventData>(Chaos::EEventType::Breaking, this, &UNiagaraDataInterfaceChaosDestruction::HandleBreakingEvents);
-			EventManager->RegisterHandler<Chaos::FTrailingEventData>(Chaos::EEventType::Trailing, this, &UNiagaraDataInterfaceChaosDestruction::HandleTrailingEvents);
-		}
-	}
-#endif
-
 	PushToRenderThread();
 }
 
@@ -192,10 +177,9 @@ void UNiagaraDataInterfaceChaosDestruction::BeginDestroy()
 	Super::BeginDestroy();
 
 #if WITH_CHAOS
-	if (GetWorld() && GetWorld()->GetPhysicsScene())
+	for (FSolverData& SolverData : Solvers)
 	{
-		FPhysScene* Scene = GetWorld()->GetPhysicsScene();
-		Chaos::FEventManager* EventManager = Scene->GetScene().GetSolver()->GetEventManager();
+		Chaos::FEventManager* EventManager = SolverData.Solver->GetEventManager();
 		if (EventManager)
 		{
 			EventManager->UnregisterHandler(Chaos::EEventType::Collision, this);
@@ -345,7 +329,6 @@ bool UNiagaraDataInterfaceChaosDestruction::CopyToInternal(UNiagaraDataInterface
 		DestinationChaosDestruction->FinalVelocityMagnitudeMinMax = FinalVelocityMagnitudeMinMax;
 		DestinationChaosDestruction->MaxLatency = MaxLatency;
 		DestinationChaosDestruction->DebugType = DebugType;
-		//DestinationChaosDestruction->ParticleIndexToProcess = ParticleIndexToProcess;
 		DestinationChaosDestruction->LastSpawnedPointID = LastSpawnedPointID;
 		DestinationChaosDestruction->LastSpawnTime = LastSpawnTime;
 		DestinationChaosDestruction->TimeStampOfLastProcessedData = TimeStampOfLastProcessedData;
@@ -427,7 +410,6 @@ bool UNiagaraDataInterfaceChaosDestruction::Equals(const UNiagaraDataInterface*
 		&& OtherChaosDestruction->FinalVelocityMagnitudeMinMax == FinalVelocityMagnitudeMinMax
 		&& OtherChaosDestruction->MaxLatency == MaxLatency
 		&& OtherChaosDestruction->DebugType == DebugType;
-		//&& OtherChaosDestruction->ParticleIndexToProcess == ParticleIndexToProcess;
 }
 
 int32 UNiagaraDataInterfaceChaosDestruction::PerInstanceDataSize()const
@@ -435,6 +417,22 @@ int32 UNiagaraDataInterfaceChaosDestruction::PerInstanceDataSize()const
 	return sizeof(FNDIChaosDestruction_InstanceData);
 }
 
+#if INCLUDE_CHAOS
+void UNiagaraDataInterfaceChaosDestruction::RegisterWithSolverEventManager(Chaos::FPhysicsSolver* Solver)
+{
+	Chaos::FEventManager* EventManager = Solver->GetEventManager();
+	if (EventManager)
+	{
+		Solver->SetGenerateCollisionData(true);
+		Solver->SetGenerateBreakingData(true);
+		Solver->SetGenerateTrailingData(true);
+		EventManager->RegisterHandler<Chaos::FCollisionEventData>(Chaos::EEventType::Collision, this, &UNiagaraDataInterfaceChaosDestruction::HandleCollisionEvents);
+		EventManager->RegisterHandler<Chaos::FBreakingEventData>(Chaos::EEventType::Breaking, this, &UNiagaraDataInterfaceChaosDestruction::HandleBreakingEvents);
+		EventManager->RegisterHandler<Chaos::FTrailingEventData>(Chaos::EEventType::Trailing, this, &UNiagaraDataInterfaceChaosDestruction::HandleTrailingEvents);
+	}
+}
+#endif
+
 bool UNiagaraDataInterfaceChaosDestruction::InitPerInstanceData(void* PerInstanceData, FNiagaraSystemInstance* SystemInstance)
 {
 	FNDIChaosDestruction_InstanceData* InstData = new (PerInstanceData) FNDIChaosDestruction_InstanceData();
@@ -456,8 +454,10 @@ bool UNiagaraDataInterfaceChaosDestruction::InitPerInstanceData(void* PerInstanc
 					int32 NewIdx = Solvers.Add(FSolverData());
 
 					FSolverData& SolverData = Solvers[NewIdx];
-					SolverData.PhysScene = World->PhysicsScene_Chaos;
-					SolverData.Solver = World->PhysicsScene_Chaos->GetSolver();
+					SolverData.PhysScene = World->GetPhysicsScene();
+					SolverData.Solver = SolverData.PhysScene->GetScene().GetSolver();
+
+					RegisterWithSolverEventManager(SolverData.Solver);
 				}
 			}
 		}
@@ -465,20 +465,25 @@ bool UNiagaraDataInterfaceChaosDestruction::InitPerInstanceData(void* PerInstanc
 	}
 	else
 	{
-		for (AChaosSolverActor* SolverActor : ChaosSolverActorSet)
-		{
-			if (SolverActor)
-			{
-				if (Chaos::FPhysicsSolver* Solver = SolverActor->GetSolver())
-				{
-					int32 NewIdx = Solvers.Add(FSolverData());
-
-					FSolverData& SolverData = Solvers[NewIdx];
-					SolverData.PhysScene = SolverActor->GetPhysicsScene();
-					SolverData.Solver = Solver;
-				}
-			}
-		}
+//#todo : are solver actors going to exist going forwards?
+//		for (AChaosSolverActor* SolverActor : ChaosSolverActorSet)
+//		{
+//			if (SolverActor)
+//			{
+//				if (Chaos::FPhysicsSolver* Solver = SolverActor->GetSolver())
+//				{
+//					int32 NewIdx = Solvers.Add(FSolverData());
+//
+//					FSolverData& SolverData = Solvers[NewIdx];
+//					SolverData.PhysScene = SolverActor->GetPhysicsScene();
+//					SolverData.Solver = Solver;
+//
+//#if INCLUDE_CHAOS
+//					RegisterWithSolverEventManager(Solver);
+//#endif
+//				}
+//			}
+//		}
 	}
 
 	ResetInstData(InstData);
@@ -509,6 +514,7 @@ void UNiagaraDataInterfaceChaosDestruction::DestroyPerInstanceData(void* PerInst
 	);
 }
 
+#if CHAOS_PARTICLEHANDLE_TODO
 void GetMeshExtData(FSolverData SolverData,
 					const int32 ParticleIndex,
 					const TArray<PhysicsProxyWrapper>& PhysicsProxyReverseMapping,
@@ -685,19 +691,25 @@ void GetMesPhysicalData(FSolverData SolverData,
 		}
 	}
 }
+#endif
 
 void UNiagaraDataInterfaceChaosDestruction::HandleCollisionEvents(const Chaos::FCollisionEventData& Event)
 {
 	ensure(IsInGameThread());
-	
+	Chaos::FCollisionDataArray const& CollisionDataIn = Event.CollisionData.AllCollisionsArray;
+
+	CollisionEvents.Reset();
+
 	// Copy data from Event into AllCollisionsArray
 	// Also get Boundingbox related data and SurfaceType and save it as well
 	CollisionEvents.AddUninitialized(Event.CollisionData.AllCollisionsArray.Num());
 
 	int32 Idx = 0;
-	for (Chaos::TCollisionDataExt<float, 3>& Collision : CollisionEvents)
+	for (Chaos::TCollisionData<float, 3> const& DataIn : CollisionDataIn)
 	{
-		Collision = Event.CollisionData.AllCollisionsArray[Idx];
+		auto& CopyData = CollisionEvents[Idx];
+
+		CopyData = DataIn;
 
 		// #GM: Disable this for now for perf
 		/*
@@ -710,10 +722,11 @@ void UNiagaraDataInterfaceChaosDestruction::HandleCollisionEvents(const Chaos::F
 			AllCollisionsArray[Idx].BoundingboxExtentMax,
 			AllCollisionsArray[Idx].SurfaceType);
 		*/
-		CollisionEvents[Idx].BoundingboxVolume = 1000000.f;
-		CollisionEvents[Idx].BoundingboxExtentMin = 100.f;
-		CollisionEvents[Idx].BoundingboxExtentMax = 100.f;
-		CollisionEvents[Idx].SurfaceType = 0;
+
+		CopyData.BoundingboxVolume = 1000000.f;
+		CopyData.BoundingboxExtentMin = 100.f;
+		CopyData.BoundingboxExtentMax = 100.f;
+		CopyData.SurfaceType = 0;
 
 		Idx++;
 	}
@@ -1107,7 +1120,7 @@ bool UNiagaraDataInterfaceChaosDestruction::CollisionCallback(FNDIChaosDestructi
 		if (SolverData.Solver->GetEventFilters()->IsCollisionEventEnabled() && CollisionEvents.Num() > 0 && SolverData.Solver->GetSolverTime() > 0.f && MaxNumberOfDataEntriesToSpawn > 0)
 		{
 			TArray<Chaos::TCollisionDataExt<float, 3>>& AllCollisionsArray = CollisionEvents;
-			float TimeData_MapsCreated = 0.0f;
+			float TimeData_MapsCreated = SolverData.Solver->GetSolverTime();
 
 #if STATS
 			{
@@ -1185,45 +1198,58 @@ bool UNiagaraDataInterfaceChaosDestruction::CollisionCallback(FNDIChaosDestructi
 void UNiagaraDataInterfaceChaosDestruction::HandleBreakingEvents(const Chaos::FBreakingEventData& Event)
 {
 	ensure(IsInGameThread());
+	Chaos::FBreakingDataArray const& BreakingDataIn = Event.BreakingData.AllBreakingsArray;
+
+	BreakingEvents.Reset();
+
+	// Copy data from Event
+	BreakingEvents.InsertZeroed(0, BreakingDataIn.Num());
 
-	// Copy data from *AllBreakingData_Maps.AllBreakingData into AllBreakingsArray
-	// Also get Boundingbox related data and SurfaceType and save it as well
-	BreakingEvents.InsertZeroed(0, Event.BreakingData.AllBreakingsArray.Num());
-	//UE_LOG(LogScript, Warning, TEXT("(*AllBreakingData_Maps.AllBreakingData).AllBreakingsArray.Num() = %d"), (*AllBreakingData_Maps.AllBreakingData).AllBreakingsArray.Num());
 	int32 Idx = 0;
-	for (Chaos::TBreakingDataExt<float, 3>& Breaking : BreakingEvents)
+	for (Chaos::TBreakingData<float, 3> const& DataIn : BreakingDataIn)
 	{
-		Breaking = Event.BreakingData.AllBreakingsArray[Idx];
-
-		// #GM: Disable this for now for perf
-		// TODO(mv): Temporarily re-enable this for Jon to get materials, transform and bounding box. Will optimize in the coming week. 
 		if (bGetExternalBreakingData)
 		{
-			//Chaos::TRigidTransform<float, 3> Transform;
-			//UPhysicalMaterial* Material = nullptr;
-			//GetMeshExtData(SolverData,
-			//	AllBreakingsArray[Idx].ParticleIndexMesh == INDEX_NONE ? AllBreakingsArray[Idx].ParticleIndex : AllBreakingsArray[Idx].ParticleIndexMesh,
-			//	PhysicsProxyReverseMappingArray,
-			//	ParticleIndexReverseMappingArray,
-			//	AllBreakingsArray[Idx].BoundingboxVolume,
-			//	AllBreakingsArray[Idx].BoundingboxExtentMin,
-			//	AllBreakingsArray[Idx].BoundingboxExtentMax,
-			//	AllBreakingsArray[Idx].BoundingBox,
-			//	AllBreakingsArray[Idx].SurfaceType,
-			//	Transform,
-			//	Material);
-			////UE_LOG(LogScript, Warning, TEXT("GetAllBreakingsAndMaps[%d]: SurfaceType = %d"), Idx, AllBreakingsArray[Idx].SurfaceType);
-			//AllBreakingsArray[Idx].TransformTranslation = Transform.GetTranslation();
-			//AllBreakingsArray[Idx].TransformRotation = Transform.GetRotation();
-			//AllBreakingsArray[Idx].TransformScale = Transform.GetScale3D();
-			//if (Material)
-			//{
-			//	AllBreakingsArray[Idx].PhysicalMaterialName = Material->GetFName();
-			//}
-			//else
-			//{
-			//	AllBreakingsArray[Idx].PhysicalMaterialName = FName();
-			//}
+			auto& CopyData = BreakingEvents[Idx];
+			CopyData = DataIn;
+
+			Chaos::TRigidTransform<float, 3> Transform;
+
+			// Ext Data..
+			CopyData.TransformTranslation = Transform.GetTranslation();
+			CopyData.TransformRotation = Transform.GetRotation();
+			CopyData.TransformScale = Transform.GetScale3D();
+
+			CopyData.BoundingBox = FBox(DataIn.BoundingBox.Min(), DataIn.BoundingBox.Max());
+			FVector Extents = CopyData.BoundingBox.GetSize();
+			CopyData.BoundingboxExtentMin = FMath::Min3(Extents[0], Extents[1], Extents[2]);
+			CopyData.BoundingboxExtentMax = FMath::Max3(Extents[0], Extents[1], Extents[2]);
+			CopyData.BoundingboxVolume = CopyData.BoundingBox.GetVolume();
+
+			UPhysicalMaterial* PhysicalMaterial = nullptr;
+			int32 MaterialID = DataIn.Particle->Geometry()->GetMaterialIndex(0);
+			UGeometryCollectionComponent* GeometryCollectionComponent = Solvers[0].PhysScene->GetScene().GetOwningComponent<UGeometryCollectionComponent>(DataIn.ParticleProxy);
+
+			UMaterialInterface* Material = GeometryCollectionComponent->GetMaterial(MaterialID);
+			ensure(Material);
+			if (Material)
+			{
+				PhysicalMaterial = Material->GetPhysicalMaterial();
+				ensure(PhysicalMaterial);
+				if (PhysicalMaterial)
+				{
+					CopyData.SurfaceType = PhysicalMaterial->SurfaceType;
+				}
+			}
+
+			if (Material)
+			{
+				CopyData.PhysicalMaterialName = Material->GetFName();
+			}
+			else
+			{
+				CopyData.PhysicalMaterialName = FName();
+			}
 		}
 		else
 		{
@@ -1246,7 +1272,6 @@ void UNiagaraDataInterfaceChaosDestruction::HandleBreakingEvents(const Chaos::FB
 void UNiagaraDataInterfaceChaosDestruction::FilterAllBreakings(TArray<Chaos::TBreakingDataExt<float, 3>>& AllBreakingsArray)
 {
 	if (bApplyMaterialsFilter || 
-	//	ParticleIndexToProcess != -1 ||
 		SpeedToSpawnMinMax.X > 0.f ||
 		SpeedToSpawnMinMax.Y > 0.f ||
 		MassToSpawnMinMax.X > 0.f ||
@@ -1299,7 +1324,6 @@ void UNiagaraDataInterfaceChaosDestruction::FilterAllBreakings(TArray<Chaos::TBr
 			float BreakingSpeedSquared = AllBreakingsArray[IdxBreaking].Velocity.SizeSquared();
 
 			if (!(bApplyMaterialsFilter && IsMaterialInFilter(AllBreakingsArray[IdxBreaking].PhysicalMaterialName)) ||
-				//(ParticleIndexToProcess != -1 && AllBreakingsArray[IdxBreaking].ParticleIndex != ParticleIndexToProcess) ||
 				(SpeedToSpawnMinMax.X > 0.f && SpeedToSpawnMinMax.Y < 0.f && BreakingSpeedSquared < MinSpeedToSpawnSquared) ||
 				(SpeedToSpawnMinMax.X < 0.f && SpeedToSpawnMinMax.Y > 0.f && BreakingSpeedSquared > MaxSpeedToSpawnSquared) ||
 				(SpeedToSpawnMinMax.X > 0.f && SpeedToSpawnMinMax.Y > 0.f && (BreakingSpeedSquared < MinSpeedToSpawnSquared || BreakingSpeedSquared > MaxSpeedToSpawnSquared)) ||
@@ -1592,7 +1616,7 @@ int32 UNiagaraDataInterfaceChaosDestruction::SpawnParticlesFromBreaking(FSolverD
 			}
 			else if (DebugType == EDebugTypeEnum::ChaosNiagara_DebugType_ColorByParticleIndex)
 			{
-				ParticleColor = ColorArray[Breaking.ParticleIndex % ColorArray.Num()];
+				//ParticleColor = ColorArray[Breaking.ParticleIndex % ColorArray.Num()]; #todo: ParticleIndex no longer exits
 			}
 
 			// Store principal data
@@ -1635,24 +1659,17 @@ bool UNiagaraDataInterfaceChaosDestruction::BreakingCallback(FNDIChaosDestructio
 		if (SolverData.Solver->GetEventFilters()->IsBreakingEventEnabled() && BreakingEvents.Num() > 0 && SolverData.Solver->GetSolverTime() > 0.f && MaxNumberOfDataEntriesToSpawn > 0)
 		{
 			TArray<Chaos::TBreakingDataExt<float, 3>>& AllBreakingsArray = BreakingEvents;
-			TArray<PhysicsProxyWrapper> PhysicsProxyReverseMappingArray;
-			TArray<int32> ParticleIndexReverseMappingArray;
 			TMap<IPhysicsProxyBase*, TArray<int32>> AllBreakingsIndicesByPhysicsProxyMap;
-			float TimeData_MapsCreated = 0.0f;
+			float TimeData_MapsCreated = SolverData.Solver->GetSolverTime();
 
 			{
 				size_t SizeOfAllBreakings = sizeof(Chaos::TBreakingData<float, 3>) * AllBreakingsArray.Num();
-				size_t SizeOfPhysicsProxyReverseMapping = sizeof(PhysicsProxyWrapper) * PhysicsProxyReverseMappingArray.Num();
-				size_t SizeOfParticleIndexReverseMapping = sizeof(int32) * ParticleIndexReverseMappingArray.Num();
-
 				size_t SizeOfAllBreakingsIndicesByPhysicsProxy = 0;
 				for (auto& Elem : AllBreakingsIndicesByPhysicsProxyMap)
 				{
 					SizeOfAllBreakingsIndicesByPhysicsProxy += sizeof(int32) * (Elem.Value).Num();
 				}
 				SET_MEMORY_STAT(STAT_AllBreakingsDataMemory, SizeOfAllBreakings);
-				SET_MEMORY_STAT(STAT_PhysicsProxyReverseMappingMemory, SizeOfPhysicsProxyReverseMapping);
-				SET_MEMORY_STAT(STAT_ParticleIndexReverseMappingMemory, SizeOfParticleIndexReverseMapping);
 				SET_MEMORY_STAT(STAT_AllBreakingsIndicesByPhysicsProxyMemory, SizeOfAllBreakingsIndicesByPhysicsProxy);
 			}
 
@@ -1681,7 +1698,6 @@ bool UNiagaraDataInterfaceChaosDestruction::BreakingCallback(FNDIChaosDestructio
 						TimeData_MapsCreated,
 						IdxSolver);
 
-					//UE_LOG(LogScript, Warning, TEXT("IdxBreaking = %d/%d, NumParticlesSpawned = %d"), IdxBreaking, BreakingsToSpawnArray.Num()-1, NumParticlesSpawned);
 					if (NumParticlesSpawned > 0)
 					{
 						// Get/Store Geometry/Physical Material data
@@ -1690,15 +1706,15 @@ bool UNiagaraDataInterfaceChaosDestruction::BreakingCallback(FNDIChaosDestructio
 						float Friction = 0.7f, Restitution = 0.3f, Density = 1.f;
 						if (bGetExternalBreakingData)
 						{
-							// #GM: Disable this for now for perf
+							/* #GM: Disable this for now for perf
 							GetMesPhysicalData(SolverData,
-											   BreakingsToSpawnArray[IdxBreaking].ParticleIndexMesh == INDEX_NONE ? BreakingsToSpawnArray[IdxBreaking].ParticleIndex : BreakingsToSpawnArray[IdxBreaking].ParticleIndexMesh,
+											   0, // #todo: BreakingsToSpawnArray[IdxBreaking].ParticleIndexMesh == INDEX_NONE ? BreakingsToSpawnArray[IdxBreaking].ParticleIndex : BreakingsToSpawnArray[IdxBreaking].ParticleIndexMesh,
 											   PhysicsProxyReverseMappingArray,
 											   ParticleIndexReverseMappingArray,
 											   Color,
 											   Friction,
 											   Restitution,
-											   Density);
+											   Density);*/
 						}
 						
 						for (int32 Idx = 0; Idx < NumParticlesSpawned; ++Idx)
@@ -1729,15 +1745,18 @@ bool UNiagaraDataInterfaceChaosDestruction::BreakingCallback(FNDIChaosDestructio
 void UNiagaraDataInterfaceChaosDestruction::HandleTrailingEvents(const Chaos::FTrailingEventData& Event)
 {
 	ensure(IsInGameThread());
+	Chaos::FTrailingDataArray const& TrailingDataIn = Event.TrailingData.AllTrailingsArray;
 
-	// Copy data from *AllTrailingData_Maps.AllTrailingData into AllTrailingsArray
-	// Also get Boundingbox related data and SurfaceType and save it as well
+	TrailingEvents.Reset();
+
+	// Copy data from Event
 	TrailingEvents.AddUninitialized(Event.TrailingData.AllTrailingsArray.Num());
 
 	int32 Idx = 0;
-	for (Chaos::TTrailingDataExt<float, 3>& Trailing : TrailingEvents)
+	for (Chaos::TTrailingData<float, 3> const& DataIn : TrailingDataIn)
 	{
-		Trailing = Event.TrailingData.AllTrailingsArray[Idx];
+		auto& CopyData = TrailingEvents[Idx];
+		CopyData = DataIn;
 
 		// #GM: Disable this for now for perf
 		/*
@@ -1750,10 +1769,11 @@ void UNiagaraDataInterfaceChaosDestruction::HandleTrailingEvents(const Chaos::FT
 			AllTrailingsArray[Idx].BoundingboxExtentMax,
 			AllTrailingsArray[Idx].SurfaceType);
 		*/
-		TrailingEvents[Idx].BoundingboxVolume = 1000000.f;
-		TrailingEvents[Idx].BoundingboxExtentMin = 100.f;
-		TrailingEvents[Idx].BoundingboxExtentMax = 100.f;
-		TrailingEvents[Idx].SurfaceType = 0;
+
+		CopyData.BoundingboxVolume = 1000000.f;
+		CopyData.BoundingboxExtentMin = 100.f;
+		CopyData.BoundingboxExtentMax = 100.f;
+		CopyData.SurfaceType = 0;
 
 		Idx++;
 	}
@@ -1761,8 +1781,7 @@ void UNiagaraDataInterfaceChaosDestruction::HandleTrailingEvents(const Chaos::FT
 
 void UNiagaraDataInterfaceChaosDestruction::FilterAllTrailings(TArray<Chaos::TTrailingDataExt<float, 3>>& AllTrailingsArray)
 {
-	if (/*ParticleIndexToProcess != -1 ||*/
-		SpeedToSpawnMinMax.X > 0.f ||
+	if (SpeedToSpawnMinMax.X > 0.f ||
 		SpeedToSpawnMinMax.Y > 0.f ||
 		MassToSpawnMinMax.X > 0.f ||
 		MassToSpawnMinMax.Y > 0.f ||
@@ -1791,8 +1810,7 @@ void UNiagaraDataInterfaceChaosDestruction::FilterAllTrailings(TArray<Chaos::TTr
 		{
 			float TrailingSpeedSquared = AllTrailingsArray[IdxTrailing].Velocity.SizeSquared();
 
-			if (/*(ParticleIndexToProcess != -1 && AllTrailingsArray[IdxTrailing].ParticleIndex != ParticleIndexToProcess) ||*/
-				(SpeedToSpawnMinMax.X > 0.f && SpeedToSpawnMinMax.Y < 0.f && TrailingSpeedSquared < MinSpeedToSpawnSquared) ||
+			if ((SpeedToSpawnMinMax.X > 0.f && SpeedToSpawnMinMax.Y < 0.f && TrailingSpeedSquared < MinSpeedToSpawnSquared) ||
 				(SpeedToSpawnMinMax.X < 0.f && SpeedToSpawnMinMax.Y > 0.f && TrailingSpeedSquared > MaxSpeedToSpawnSquared) ||
 				(SpeedToSpawnMinMax.X > 0.f && SpeedToSpawnMinMax.Y > 0.f && (TrailingSpeedSquared < MinSpeedToSpawnSquared || TrailingSpeedSquared > MaxSpeedToSpawnSquared)) ||
 				(MassToSpawnMinMax.X > 0.f && MassToSpawnMinMax.Y < 0.f && AllTrailingsArray[IdxTrailing].Mass < MassToSpawnMinMax.X) ||
@@ -1989,7 +2007,7 @@ int32 UNiagaraDataInterfaceChaosDestruction::SpawnParticlesFromTrailing(FSolverD
 			}
 			else if (DebugType == EDebugTypeEnum::ChaosNiagara_DebugType_ColorByParticleIndex)
 			{
-				//ParticleColor = ColorArray[Trailing.ParticleIndex % ColorArray.Num()];
+				//ParticleColor = ColorArray[Trailing.ParticleIndex % ColorArray.Num()]; #todo: ParticleIndex no longer exists
 			}
 
 			// Store principal data
@@ -2031,24 +2049,17 @@ bool UNiagaraDataInterfaceChaosDestruction::TrailingCallback(FNDIChaosDestructio
 		if (SolverData.Solver->GetEventFilters()->IsTrailingEventEnabled() && TrailingEvents.Num() > 0 && SolverData.Solver->GetSolverTime() > 0.f && MaxNumberOfDataEntriesToSpawn > 0)
 		{
 			TArray<Chaos::TTrailingDataExt<float, 3>>& AllTrailingsArray = TrailingEvents;
-			TArray<PhysicsProxyWrapper> PhysicsProxyReverseMappingArray;
-			TArray<int32> ParticleIndexReverseMappingArray;
 			TMap<IPhysicsProxyBase*, TArray<int32>> AllTrailingsIndicesByPhysicsProxyMap;
-			float TimeData_MapsCreated = 0.0f;
+			float TimeData_MapsCreated = SolverData.Solver->GetSolverTime();
 
 			{
 				size_t SizeOfAllTrailings = sizeof(Chaos::TTrailingData<float, 3>) * AllTrailingsArray.Num();
-				size_t SizeOfPhysicsProxyReverseMapping = sizeof(PhysicsProxyWrapper) * PhysicsProxyReverseMappingArray.Num();
-				size_t SizeOfParticleIndexReverseMapping = sizeof(int32) * ParticleIndexReverseMappingArray.Num();
-
 				size_t SizeOfAllTrailingsIndicesByPhysicsProxy = 0;
 				for (auto& Elem : AllTrailingsIndicesByPhysicsProxyMap)
 				{
 					SizeOfAllTrailingsIndicesByPhysicsProxy += sizeof(int32) * (Elem.Value).Num();
 				}
 				SET_MEMORY_STAT(STAT_AllTrailingsDataMemory, SizeOfAllTrailings);
-				SET_MEMORY_STAT(STAT_PhysicsProxyReverseMappingMemory, SizeOfPhysicsProxyReverseMapping);
-				SET_MEMORY_STAT(STAT_ParticleIndexReverseMappingMemory, SizeOfParticleIndexReverseMapping);
 				SET_MEMORY_STAT(STAT_AllTrailingsIndicesByPhysicsProxyMemory, SizeOfAllTrailingsIndicesByPhysicsProxy);
 			}
 
diff --git a/Engine/Source/Runtime/Experimental/Chaos/Private/Chaos/PBDConstraintGraph.cpp b/Engine/Source/Runtime/Experimental/Chaos/Private/Chaos/PBDConstraintGraph.cpp
--- a/Engine/Source/Runtime/Experimental/Chaos/Private/Chaos/PBDConstraintGraph.cpp
+++ b/Engine/Source/Runtime/Experimental/Chaos/Private/Chaos/PBDConstraintGraph.cpp
@@ -65,7 +65,8 @@ int32 FPBDConstraintGraph::ReserveParticles(const int32 Num)
 
 void FPBDConstraintGraph::ParticleAdd(TGeometryParticleHandle<FReal, 3>* AddedParticle)
 {
-	if (ensure(!ParticleToNodeIndex.Contains(AddedParticle)))
+	// GC Code creates particle then enables particle so end up calling this twice triggering ensure
+	if (/*ensure*/ (!ParticleToNodeIndex.Contains(AddedParticle)))
 	{
 		const int32 NewNodeIndex = GetNextNodeIndex();
 		FGraphNode& Node = Nodes[NewNodeIndex];
@@ -907,4 +908,3 @@ bool FPBDConstraintGraph::CheckIslands(const TArray<TGeometryParticleHandle<FRea
 
 	return bIsValid;
 }
-
diff --git a/Engine/Source/Runtime/Experimental/Chaos/Private/Chaos/PBDRigidClustering.cpp b/Engine/Source/Runtime/Experimental/Chaos/Private/Chaos/PBDRigidClustering.cpp
--- a/Engine/Source/Runtime/Experimental/Chaos/Private/Chaos/PBDRigidClustering.cpp
+++ b/Engine/Source/Runtime/Experimental/Chaos/Private/Chaos/PBDRigidClustering.cpp
@@ -745,7 +745,8 @@ namespace Chaos
 					{
 						const int32 NewIdx = MAllClusterBreakings.Add(TBreakingData<float, 3>());
 						TBreakingData<float, 3>& ClusterBreak = MAllClusterBreakings[NewIdx];
-						ClusterBreak.Particle = Child->GTGeometryParticle();
+						ClusterBreak.Particle = Child;
+						ClusterBreak.ParticleProxy = nullptr;
 						ClusterBreak.Location = Child->X();
 						ClusterBreak.Velocity = Child->V();
 						ClusterBreak.AngularVelocity = Child->W();
@@ -1078,7 +1079,8 @@ namespace Chaos
 						{
 							int32 NewIdx = MAllClusterBreakings.Add(TBreakingData<float, 3>());
 							TBreakingData<float, 3>& ClusterBreak = MAllClusterBreakings[NewIdx];
-							ClusterBreak.Particle = ClusteredParticle->GTGeometryParticle();
+							ClusterBreak.Particle = ClusteredParticle;
+							ClusterBreak.ParticleProxy = nullptr;
 							ClusterBreak.Location = ClusteredParticle->X();
 							ClusterBreak.Velocity = ClusteredParticle->V();
 							ClusterBreak.AngularVelocity = ClusteredParticle->W();
@@ -1632,7 +1634,8 @@ namespace Chaos
 		{
 			const int32 NewIdx = MAllClusterBreakings.Add(TBreakingData<float, 3>());
 			TBreakingData<float, 3>& ClusterBreak = MAllClusterBreakings[NewIdx];
-			ClusterBreak.Particle = Particle->GTGeometryParticle();
+			ClusterBreak.Particle = Particle;
+			ClusterBreak.ParticleProxy = nullptr;
 			ClusterBreak.Location = Particle->X();
 			ClusterBreak.Velocity = Particle->V();
 			ClusterBreak.AngularVelocity = Particle->W();
diff --git a/Engine/Source/Runtime/Experimental/Chaos/Public/Chaos/ExternalCollisionData.h b/Engine/Source/Runtime/Experimental/Chaos/Public/Chaos/ExternalCollisionData.h
--- a/Engine/Source/Runtime/Experimental/Chaos/Public/Chaos/ExternalCollisionData.h
+++ b/Engine/Source/Runtime/Experimental/Chaos/Public/Chaos/ExternalCollisionData.h
@@ -1,10 +1,11 @@
 // Copyright Epic Games, Inc. All Rights Reserved.
 #pragma once
 #include "Box.h"
-#include "Chaos/ParticleHandleFwd.h"
+#include "Chaos/Framework/PhysicsProxyBase.h"
 #include "Chaos/PBDConstraintContainer.h"
+#include "Chaos/ParticleHandleFwd.h"
 #include "Chaos/Vector.h"
-#include "Chaos/Framework/PhysicsProxyBase.h"
+
 
 class UPhysicalMaterial;
 
@@ -31,26 +32,9 @@ namespace Chaos
 			, Levelset(nullptr)
 			, ParticleProxy(nullptr)
 		    , LevelsetProxy(nullptr)
-			, ParticleIndexMesh(INDEX_NONE)
-			, LevelsetIndexMesh(INDEX_NONE)
 		{}
 
-		TCollisionData(TVector<T, d> InLocation
-			, TVector<T, d> InAccumulatedImpulse
-			, TVector<T, d> InNormal
-			, TVector<T, d> InVelocity1
-			, TVector<T, d> InVelocity2
-			, TVector<T, d> InAngularVelocity1
-			, TVector<T, d> InAngularVelocity2
-			, T InMass1
-			, T InMass2
-			, T InPenetrationDepth
-			, TGeometryParticle<T, d>* InParticle
-			, TGeometryParticle<T, d>* InLevelset
-			, IPhysicsProxyBase* InParticleProxy
-			, IPhysicsProxyBase* InLevelsetProxy
-			, int32 InParticleIndexMesh
-			, int32 InLevelsetIndexMesh)
+		TCollisionData(TVector<T, d> InLocation, TVector<T, d> InAccumulatedImpulse, TVector<T, d> InNormal, TVector<T, d> InVelocity1, TVector<T, d> InVelocity2, TVector<T, d> InAngularVelocity1, TVector<T, d> InAngularVelocity2, T InMass1, T InMass2, T InPenetrationDepth, TGeometryParticle<T, d>* InParticle, TGeometryParticle<T, d>* InLevelset, IPhysicsProxyBase* InParticleProxy, IPhysicsProxyBase* InLevelsetProxy)
 			: Location(InLocation)
 			, AccumulatedImpulse(InAccumulatedImpulse)
 			, Normal(InNormal)
@@ -65,8 +49,6 @@ namespace Chaos
 			, Levelset(InLevelset)
 			, ParticleProxy(InParticleProxy)
 		    , LevelsetProxy(InLevelsetProxy)
-			, ParticleIndexMesh(InParticleIndexMesh)
-			, LevelsetIndexMesh(InLevelsetIndexMesh)
 		{}
 
 		bool IsValid() { return (ParticleProxy && LevelsetProxy); }
@@ -82,13 +64,8 @@ namespace Chaos
 		TGeometryParticle<T, d>* Levelset;
 		IPhysicsProxyBase* ParticleProxy;
 		IPhysicsProxyBase* LevelsetProxy;
-		// @todo(ccaulfield): CHAOS_PARTICLEHANDLE_TODO
-		int32 ParticleIndexMesh, LevelsetIndexMesh; // If ParticleIndex points to a cluster then this index will point to an actual mesh in the cluster
-													// It is important to be able to get extra data from the component
-
 	};
 
-
 	/*
 	CollisionData used in the subsystems
 	*/
@@ -107,31 +84,16 @@ namespace Chaos
 			, Mass2((T)0.0)
 			, Particle(nullptr)
 			, Levelset(nullptr)
-			, ParticleIndexMesh(INDEX_NONE)
-			, LevelsetIndexMesh(INDEX_NONE)
+		    , ParticleProxy(nullptr)
+		    , LevelsetProxy(nullptr)
 			, BoundingboxVolume((T)-1.0)
 			, BoundingboxExtentMin((T)-1.0)
 			, BoundingboxExtentMax((T)-1.0)
 			, SurfaceType(-1)
 		{}
 
-		TCollisionDataExt(TVector<T, d> InLocation
-			, TVector<T, d> InAccumulatedImpulse
-			, TVector<T, d> InNormal
-			, TVector<T, d> InVelocity1
-			, TVector<T, d> InVelocity2
-			, TVector<T, d> InAngularVelocity1
-			, TVector<T, d> InAngularVelocity2
-			, T InMass1
-			, T InMass2
-			, TGeometryParticle<T, d>* InParticle
-			, TGeometryParticle<T, d>* InLevelset
-			, int32 InParticleIndexMesh
-			, int32 InLevelsetIndexMesh
-			, float InBoundingboxVolume
-			, float InBoundingboxExtentMin
-			, float InBoundingboxExtentMax
-			, int32 InSurfaceType)
+		TCollisionDataExt(
+		    TVector<T, d> InLocation, TVector<T, d> InAccumulatedImpulse, TVector<T, d> InNormal, TVector<T, d> InVelocity1, TVector<T, d> InVelocity2, TVector<T, d> InAngularVelocity1, TVector<T, d> InAngularVelocity2, T InMass1, T InMass2, TGeometryParticle<T, d>* InParticle, TGeometryParticle<T, d>* InLevelset, IPhysicsProxyBase* ParticleProxy, IPhysicsProxyBase* LevelSetProxy, float InBoundingboxVolume, float InBoundingboxExtentMin, float InBoundingboxExtentMax, int32 InSurfaceType)
 			: Location(InLocation)
 			, AccumulatedImpulse(InAccumulatedImpulse)
 			, Normal(InNormal)
@@ -143,8 +105,8 @@ namespace Chaos
 			, Mass2(InMass2)
 			, Particle(InParticle)
 			, Levelset(InLevelset)
-			, ParticleIndexMesh(InParticleIndexMesh)
-			, LevelsetIndexMesh(InLevelsetIndexMesh)
+		    , ParticleProxy(InParticleProxy)
+		    , LevelsetProxy(InLevelsetProxy)
 			, BoundingboxVolume(InBoundingboxVolume)
 			, BoundingboxExtentMin(InBoundingboxExtentMin)
 			, BoundingboxExtentMax(InBoundingboxExtentMax)
@@ -163,13 +125,14 @@ namespace Chaos
 			, Mass2(InCollisionData.Mass2)
 			, Particle(InCollisionData.Particle)
 			, Levelset(InCollisionData.Levelset)
-			, ParticleIndexMesh(InCollisionData.ParticleIndexMesh)
-			, LevelsetIndexMesh(InCollisionData.LevelsetIndexMesh)
+		    , ParticleProxy(InCollisionData.ParticleProxy)
+		    , LevelsetProxy(InCollisionData.LevelsetProxy)
 			, BoundingboxVolume((T)-1.0)
 			, BoundingboxExtentMin((T)-1.0)
 			, BoundingboxExtentMax((T)-1.0)
 			, SurfaceType(-1)
-		{}
+		{
+		}
 
 		TVector<T, d> Location;
 		TVector<T, d> AccumulatedImpulse;
@@ -179,7 +142,8 @@ namespace Chaos
 		T Mass1, Mass2;
 		TGeometryParticle<T, d>* Particle;
 		TGeometryParticle<T, d>* Levelset;
-		int32 ParticleIndexMesh, LevelsetIndexMesh;
+		IPhysicsProxyBase* ParticleProxy;
+		IPhysicsProxyBase* LevelsetProxy;
 		float BoundingboxVolume;
 		float BoundingboxExtentMin, BoundingboxExtentMax;
 		int32 SurfaceType;
@@ -193,28 +157,23 @@ namespace Chaos
 	{
 		TBreakingData()
 			: Particle(nullptr)
+		    , ParticleProxy(nullptr)
 			, Location(TVector<T, d>((T)0.0))
 			, Velocity(TVector<T, d>((T)0.0))
 			, AngularVelocity(TVector<T, d>((T)0.0))
 			, Mass((T)0.0)
-			, ParticleIndex(INDEX_NONE)
-			, ParticleIndexMesh(INDEX_NONE)
 			, BoundingBox(TAABB<T, d>(TVector<T, d>((T)0.0), TVector<T, d>((T)0.0)))
 		{}
 
-		TGeometryParticle<T, d>* Particle; // 64 bits go first for better packing
+		TGeometryParticleHandle<T, d>* Particle;
+		IPhysicsProxyBase* ParticleProxy;
 		TVector<T, d> Location;
 		TVector<T, d> Velocity;
 		TVector<T, d> AngularVelocity;
 		T Mass;
-		int32 ParticleIndex; //#todo: remove this in favor of TGeometryParticle?
-		int32 ParticleIndexMesh; // If ParticleIndex points to a cluster then this index will point to an actual mesh in the cluster
-								 // It is important to be able to get extra data from the component
 		Chaos::TAABB<T, d> BoundingBox;
 	};
 
-
-
 	/*
 	BreakingData used in the subsystems
 	*/
@@ -226,8 +185,8 @@ namespace Chaos
 			, Velocity(TVector<T, d>((T)0.0))
 			, AngularVelocity(TVector<T, d>((T)0.0))
 			, Mass((T)0.0)
-			, ParticleIndex(INDEX_NONE)
-			, ParticleIndexMesh(INDEX_NONE)
+		    , Particle(nullptr)
+		    , ParticleProxy(nullptr)
 			, BoundingboxVolume((T)-1.0)
 			, BoundingboxExtentMin((T)-1.0)
 			, BoundingboxExtentMax((T)-1.0)
@@ -238,8 +197,8 @@ namespace Chaos
 			, TVector<T, d> InVelocity
 			, TVector<T, d> InAngularVelocity
 			, T InMass
-			, int32 InParticleIndex
-			, int32 InParticleIndexMesh
+						, TGeometryParticleHandle<T, d>* Particle
+						, IPhysicsProxyBase* ParticleProxy
 			, float InBoundingboxVolume
 			, float InBoundingboxExtentMin
 			, float InBoundingboxExtentMax
@@ -248,8 +207,8 @@ namespace Chaos
 			, Velocity(InVelocity)
 			, AngularVelocity(InAngularVelocity)
 			, Mass(InMass)
-			, ParticleIndex(InParticleIndex)
-			, ParticleIndexMesh(InParticleIndexMesh)
+		    , Particle(InParticle)
+		    , ParticleProxy(InParticleProxy)
 			, BoundingboxVolume(InBoundingboxVolume)
 			, BoundingboxExtentMin(InBoundingboxExtentMin)
 			, BoundingboxExtentMax(InBoundingboxExtentMax)
@@ -261,21 +220,21 @@ namespace Chaos
 			, Velocity(InBreakingData.Velocity)
 			, AngularVelocity(InBreakingData.AngularVelocity)
 			, Mass(InBreakingData.Mass)
-			, ParticleIndex(InBreakingData.ParticleIndex)
-			, ParticleIndexMesh(InBreakingData.ParticleIndexMesh)
+		    , Particle(InBreakingData.Particle)
+		    , ParticleProxy(InBreakingData.ParticleProxy)
 			, BoundingboxVolume((T)-1.0)
 			, BoundingboxExtentMin((T)-1.0)
 			, BoundingboxExtentMax((T)-1.0)
 			, SurfaceType(-1)
-		{}
+		{
+		}
 
 		TVector<T, d> Location;
 		TVector<T, d> Velocity;
 		TVector<T, d> AngularVelocity;
 		T Mass;
-		int32 ParticleIndex;
-		int32 ParticleIndexMesh; // If ParticleIndex points to a cluster then this index will point to an actual mesh in the cluster
-								 // It is important to be able to get extra data from the component
+		TGeometryParticleHandle<T, d>* Particle;
+		IPhysicsProxyBase* ParticleProxy;
 		float BoundingboxVolume;
 		float BoundingboxExtentMin, BoundingboxExtentMax;
 		int32 SurfaceType;
@@ -304,27 +263,17 @@ namespace Chaos
 			, AngularVelocity(TVector<T, d>((T)0.0))
 			, Mass((T)0.0)
 			, Particle(nullptr)
-			, ParticleIndexMesh(INDEX_NONE)
+		    , ParticleProxy(nullptr)
 			, BoundingBox(TAABB<T, d>(TVector<T, d>((T)0.0), TVector<T, d>((T)0.0)))
 		{}
 
-		TTrailingData(TVector<T, d> InLocation
-			, TVector<T, d> InVelocity
-			, TVector<T, d> InAngularVelocity
-			, T InMass
-			, TGeometryParticle<T, d>* InParticle
-			, int32 InParticleIndexMesh
-			, float InBoundingboxVolume
-			, float InBoundingboxExtentMin
-			, float InBoundingboxExtentMax
-			, int32 InSurfaceType
-			, Chaos::TAABB<T, d>& InBoundingBox)
+		TTrailingData(TVector<T, d> InLocation, TVector<T, d> InVelocity, TVector<T, d> InAngularVelocity, T InMass, TGeometryParticleHandle<T, d>* InParticle
+		    , float InBoundingboxVolume, float InBoundingboxExtentMin, float InBoundingboxExtentMax, int32 InSurfaceType, Chaos::TAABB<T, d>& InBoundingBox)
 			: Location(InLocation)
 			, Velocity(InVelocity)
 			, AngularVelocity(InAngularVelocity)
 			, Mass(InMass)
 			, Particle(InParticle)
-			, ParticleIndexMesh(InParticleIndexMesh)
 			, BoundingBox(InBoundingBox)
 		{}
 
@@ -332,9 +281,8 @@ namespace Chaos
 		TVector<T, d> Velocity;
 		TVector<T, d> AngularVelocity;
 		T Mass;
-		TGeometryParticle<T, d>* Particle;
-		int32 ParticleIndexMesh; // If ParticleIndex points to a cluster then this index will point to an actual mesh in the cluster
-								 // It is important to be able to get extra data from the component
+		TGeometryParticleHandle<T, d>* Particle;
+		IPhysicsProxyBase* ParticleProxy;
 		Chaos::TAABB<T, d> BoundingBox;
 
 		friend inline uint32 GetTypeHash(const TTrailingData& Other)
@@ -360,7 +308,6 @@ namespace Chaos
 			, AngularVelocity(TVector<T, d>((T)0.0))
 			, Mass((T)0.0)
 			, Particle(nullptr)
-			, ParticleIndexMesh(INDEX_NONE)
 			, BoundingboxVolume((T)-1.0)
 			, BoundingboxExtentMin((T)-1.0)
 			, BoundingboxExtentMax((T)-1.0)
@@ -371,8 +318,8 @@ namespace Chaos
 			, TVector<T, d> InVelocity
 			, TVector<T, d> InAngularVelocity
 			, T InMass
-			, TGeometryParticle<T, d>* InParticle
-			, int32 InParticleIndexMesh
+				, TGeometryParticleHandle<T, d>* InParticle
+				, IPhysicsProxyBase* ParticleProxy
 			, float InBoundingboxVolume
 			, float InBoundingboxExtentMin
 			, float InBoundingboxExtentMax
@@ -382,7 +329,7 @@ namespace Chaos
 			, AngularVelocity(InAngularVelocity)
 			, Mass(InMass)
 			, Particle(InParticle)
-			, ParticleIndexMesh(InParticleIndexMesh)
+		    , ParticleProxy(InParticleProxy)
 			, BoundingboxVolume(InBoundingboxVolume)
 			, BoundingboxExtentMin(InBoundingboxExtentMin)
 			, BoundingboxExtentMax(InBoundingboxExtentMax)
@@ -395,7 +342,7 @@ namespace Chaos
 			, AngularVelocity(InTrailingData.AngularVelocity)
 			, Mass(InTrailingData.Mass)
 			, Particle(InTrailingData.Particle)
-			, ParticleIndexMesh(InTrailingData.ParticleIndexMesh)
+		    , ParticleProxy(InTrailingData.ParticleProxy)
 			, BoundingboxVolume((T)-1.0)
 			, BoundingboxExtentMin((T)-1.0)
 			, BoundingboxExtentMax((T)-1.0)
@@ -406,8 +353,9 @@ namespace Chaos
 		TVector<T, d> Velocity;
 		TVector<T, d> AngularVelocity;
 		T Mass;
-		TGeometryParticle<T, d>* Particle;
-		int32 ParticleIndexMesh; // If ParticleIndex points to a cluster then this index will point to an actual mesh in the cluster
+		TGeometryParticleHandle<T, d>* Particle;
+		IPhysicsProxyBase* ParticleProxy;
+		//	int32 ParticleIndexMesh; // If ParticleIndex points to a cluster then this index will point to an actual mesh in the cluster
 								 // It is important to be able to get extra data from the component
 		float BoundingboxVolume;
 		float BoundingboxExtentMin, BoundingboxExtentMax;
@@ -424,7 +372,6 @@ namespace Chaos
 		}
 	};
 
-
 	template<class T, int d>
 	struct TSleepingData
 	{
@@ -434,8 +381,7 @@ namespace Chaos
 		{}
 
 		TSleepingData(
-			TGeometryParticle<T, d>* InParticle
-			, bool InSleeping)
+		    TGeometryParticle<T, d>* InParticle, bool InSleeping)
 			: Particle(InParticle)
 			, Sleeping(InSleeping)
 		{}
@@ -444,5 +390,6 @@ namespace Chaos
 		bool Sleeping;	// if !Sleeping == Awake
 	};
 
-
 }
+
+
diff --git a/Engine/Source/Runtime/Experimental/ChaosSolverEngine/Private/Chaos/ChaosGameplayEventDispatcher.cpp b/Engine/Source/Runtime/Experimental/ChaosSolverEngine/Private/Chaos/ChaosGameplayEventDispatcher.cpp
--- a/Engine/Source/Runtime/Experimental/ChaosSolverEngine/Private/Chaos/ChaosGameplayEventDispatcher.cpp
+++ b/Engine/Source/Runtime/Experimental/ChaosSolverEngine/Private/Chaos/ChaosGameplayEventDispatcher.cpp
@@ -388,9 +388,9 @@ void UChaosGameplayEventDispatcher::HandleBreakingEvents(const Chaos::FBreakingE
 		{
 			for (Chaos::TBreakingData<float, 3> const& BreakingDataItem : BreakingData)
 			{	
-				if (BreakingDataItem.Particle && (BreakingDataItem.Particle->GetProxy()))
+				if (BreakingDataItem.Particle && (BreakingDataItem.ParticleProxy))
 				{
-					UPrimitiveComponent* const PrimComp = Cast<UPrimitiveComponent>(BreakingDataItem.Particle->GetProxy()->GetOwner());
+					UPrimitiveComponent* const PrimComp = Cast<UPrimitiveComponent>(BreakingDataItem.ParticleProxy->GetOwner());
 					if (PrimComp && BreakEventRegistrations.Contains(PrimComp))
 					{
 						// queue them up so we can release the physics data before trigging BP events
diff --git a/Engine/Source/Runtime/Experimental/ChaosSolvers/Private/EventDefaults.cpp b/Engine/Source/Runtime/Experimental/ChaosSolvers/Private/EventDefaults.cpp
--- a/Engine/Source/Runtime/Experimental/ChaosSolvers/Private/EventDefaults.cpp
+++ b/Engine/Source/Runtime/Experimental/ChaosSolvers/Private/EventDefaults.cpp
@@ -320,13 +320,10 @@ namespace Chaos
 				return;
 
 			FBreakingDataArray& AllBreakingDataArray = BreakingEventData.BreakingData.AllBreakingsArray;
-			TMap<IPhysicsProxyBase*, TArray<int32>>& AllBreakingIndicesByPhysicsProxy = BreakingEventData.PhysicsProxyToBreakingIndices.PhysicsProxyToIndicesMap;
 
 			AllBreakingDataArray.Reset();
-			AllBreakingIndicesByPhysicsProxy.Reset();
 
 			BreakingEventData.BreakingData.TimeCreated = Solver->MTime;
-			BreakingEventData.PhysicsProxyToBreakingIndices.TimeCreated = Solver->MTime;
 
 			const FPBDRigidsSolver::FPBDRigidsEvolution* Evolution = Solver->GetEvolution();
 			const TPBDRigidParticles<float, 3>& Particles = Evolution->GetParticles().GetDynamicParticles();
@@ -343,21 +340,25 @@ namespace Chaos
 				{
 					// Since Clustered GCs can be unioned the particleIndex representing the union 
 					// is not associated with a PhysicsProxy
-					if(AllBreakingsArray[Idx].Particle->GetProxy() != nullptr)
+					TPBDRigidParticleHandle<float, 3>* PBDRigid = AllBreakingsArray[Idx].Particle->CastToRigidParticle();
+					if(PBDRigid)
 					{
 						if(ensure(!AllBreakingsArray[Idx].Location.ContainsNaN() &&
-							!Particles.V(AllBreakingsArray[Idx].ParticleIndex).ContainsNaN() &&
-							!Particles.W(AllBreakingsArray[Idx].ParticleIndex).ContainsNaN()))
+							!PBDRigid->V().ContainsNaN() &&
+							!PBDRigid->W().ContainsNaN()))
 						{
 							TBreakingData<float, 3> BreakingData;
 							BreakingData.Location = AllBreakingsArray[Idx].Location;
-							BreakingData.Velocity = Particles.V(AllBreakingsArray[Idx].ParticleIndex);
-							BreakingData.AngularVelocity = Particles.W(AllBreakingsArray[Idx].ParticleIndex);
-							BreakingData.Mass = Particles.M(AllBreakingsArray[Idx].ParticleIndex);
-							BreakingData.ParticleIndex = AllBreakingsArray[Idx].ParticleIndex;
-							if(Particles.Geometry(Idx)->HasBoundingBox())
+							BreakingData.Velocity = PBDRigid->V();
+							BreakingData.AngularVelocity = PBDRigid->W();
+							BreakingData.Mass = PBDRigid->M();
+							BreakingData.Particle = PBDRigid;
+							BreakingData.ParticleProxy = Solver->GetProxies(PBDRigid->Handle()) && Solver->GetProxies(PBDRigid->Handle())->Array().Num() ?
+								Solver->GetProxies(PBDRigid->Handle())->Array().operator[](0) : nullptr; // @todo(chaos) : Iterate all proxies
+							
+							if(PBDRigid->Geometry()->HasBoundingBox())
 							{
-								BreakingData.BoundingBox = Particles.Geometry(Idx)->BoundingBox();;
+								BreakingData.BoundingBox = PBDRigid->Geometry()->BoundingBox();
 							}
 
 							const FSolverBreakingEventFilter* SolverBreakingEventFilter = Solver->GetEventFilters()->GetBreakingFilter();
@@ -367,19 +368,15 @@ namespace Chaos
 								TBreakingData<float, 3>& BreakingDataArrayItem = AllBreakingDataArray[NewIdx];
 								BreakingDataArrayItem = BreakingData;
 
+#if 0 // #todo
 								// If AllBreakingsArray[Idx].ParticleIndex is a cluster store an index for a mesh in this cluster
 								if(ClusterIdsArray[AllBreakingsArray[Idx].ParticleIndex].NumChildren > 0)
 								{
-#if 0 // #todo
 									int32 ParticleIndexMesh = GetParticleIndexMesh(ParentToChildrenMap, AllBreakingsArray[Idx].ParticleIndex);
 									ensure(ParticleIndexMesh != INDEX_NONE);
 									BreakingDataArrayItem.ParticleIndexMesh = ParticleIndexMesh;
-#endif
 								}
-
-								// Add to AllBreakingsIndicesByPhysicsProxy
-								IPhysicsProxyBase* PhysicsProxy = AllBreakingsArray[Idx].Particle->GetProxy();
-								AllBreakingIndicesByPhysicsProxy.FindOrAdd(PhysicsProxy).Add(NewIdx);
+#endif
 							}
 						}
 					}
@@ -406,13 +403,10 @@ namespace Chaos
 			const TMap<uint32, TUniquePtr<TArray<uint32>>>& ParentToChildrenMap = Evolution->GetRigidClustering().GetChildrenMap();
 #endif
 			auto& AllTrailingsDataArray = TrailingEventData.TrailingData.AllTrailingsArray;
-			auto& AllTrailingsIndicesByPhysicsProxy = TrailingEventData.PhysicsProxyToTrailingIndices.PhysicsProxyToIndicesMap;
 
 			AllTrailingsDataArray.Reset();
-			AllTrailingsIndicesByPhysicsProxy.Reset();
 
 			TrailingEventData.TrailingData.TimeCreated = Solver->MTime;
-			TrailingEventData.PhysicsProxyToTrailingIndices.TimeCreated = Solver->MTime;
 
 			for (auto& ActiveParticle : Evolution->GetParticles().GetActiveParticlesView())
 			{
@@ -455,18 +449,6 @@ namespace Chaos
 									TrailingDataArrayItem.ParticleIndexMesh = ParticleIndexMesh;
 								}
 #endif
-
-								// Add to AllTrailingsIndicesByPhysicsProxy
-								if (const TSet<IPhysicsProxyBase*>* Proxies = Solver->GetProxies(ActiveParticle.Handle()))
-								{
-									for (IPhysicsProxyBase* Proxy : *Proxies)
-									{
-										if (Proxy != nullptr)
-										{
-											AllTrailingsIndicesByPhysicsProxy.FindOrAdd(Proxy).Add(NewIdx);
-										}
-									}
-								}
 							}
 						}
 					}
@@ -486,9 +468,7 @@ namespace Chaos
 			const FPBDRigidsSolver::FPBDRigidsEvolution* Evolution = Solver->GetEvolution();
 
 			FSleepingDataArray& EventSleepDataArray = SleepingEventData.SleepingData;
-			TMap<IPhysicsProxyBase*, TArray<int32>>& AllSleepIndicesByPhysicsProxy = SleepingEventData.PhysicsProxyToSleepingIndices.PhysicsProxyToIndicesMap;
 			EventSleepDataArray.Reset();
-			AllSleepIndicesByPhysicsProxy.Reset();
 
 			Chaos::FPBDRigidsSolver* NonConstSolver = (Chaos::FPBDRigidsSolver*)(Solver);
 
@@ -509,8 +489,6 @@ namespace Chaos
 								TSleepingData<float, 3>& SleepingDataArrayItem = EventSleepDataArray[NewIdx];
 								SleepingDataArrayItem.Particle = Particle;
 								SleepingDataArrayItem.Sleeping = SleepData.Sleeping;
-
-								AllSleepIndicesByPhysicsProxy.FindOrAdd(Proxy).Add(NewIdx);
 							}
 						}
 					}
@@ -523,4 +501,4 @@ namespace Chaos
 
 		});
 	}
-}
\ No newline at end of file
+}
diff --git a/Engine/Source/Runtime/Experimental/ChaosSolvers/Public/EventsData.h b/Engine/Source/Runtime/Experimental/ChaosSolvers/Public/EventsData.h
--- a/Engine/Source/Runtime/Experimental/ChaosSolvers/Public/EventsData.h
+++ b/Engine/Source/Runtime/Experimental/ChaosSolvers/Public/EventsData.h
@@ -86,7 +86,6 @@ namespace Chaos
 		FBreakingEventData() {}
 
 		FAllBreakingData BreakingData;
-		FIndicesByPhysicsProxy PhysicsProxyToBreakingIndices;
 	};
 
 	/* Trailing */
@@ -112,7 +111,6 @@ namespace Chaos
 		FTrailingEventData() {}
 
 		FAllTrailingData TrailingData;
-		FIndicesByPhysicsProxy PhysicsProxyToTrailingIndices;
 	};
 
 	struct FSleepingEventData
@@ -120,7 +118,6 @@ namespace Chaos
 		FSleepingEventData() {}
 
 		FSleepingDataArray SleepingData;
-		FIndicesByPhysicsProxy PhysicsProxyToSleepingIndices;
 	};
 
 }
diff --git a/Engine/Source/Runtime/Experimental/ChaosSolvers/Public/PBDRigidsSolver.h b/Engine/Source/Runtime/Experimental/ChaosSolvers/Public/PBDRigidsSolver.h
--- a/Engine/Source/Runtime/Experimental/ChaosSolvers/Public/PBDRigidsSolver.h
+++ b/Engine/Source/Runtime/Experimental/ChaosSolvers/Public/PBDRigidsSolver.h
@@ -266,9 +266,7 @@ namespace Chaos
 		void SetGenerateBreakingData(bool bDoGenerate)
 		{
 			GetEventFilters()->SetGenerateBreakingEvents(bDoGenerate);
-#if TODO_REIMPLEMENT_RIGID_CLUSTERING
-			GetRigidClustering().SetGenerateClusterBreaking(GenerateBreakingEventsEnabled);
-#endif
+			GetEvolution()->GetRigidClustering().SetGenerateClusterBreaking(bDoGenerate);
 		}
 		void SetGenerateTrailingData(bool bDoGenerate) { GetEventFilters()->SetGenerateTrailingEvents(bDoGenerate); }
 		void SetCollisionFilterSettings(const FSolverCollisionFilterSettings& InCollisionFilterSettings) { GetEventFilters()->GetCollisionFilter()->UpdateFilterSettings(InCollisionFilterSettings); }
diff --git a/Engine/Source/Runtime/Experimental/GeometryCollectionEngine/Private/ChaosBreakingEventFilter.cpp b/Engine/Source/Runtime/Experimental/GeometryCollectionEngine/Private/ChaosBreakingEventFilter.cpp
--- a/Engine/Source/Runtime/Experimental/GeometryCollectionEngine/Private/ChaosBreakingEventFilter.cpp
+++ b/Engine/Source/Runtime/Experimental/GeometryCollectionEngine/Private/ChaosBreakingEventFilter.cpp
@@ -37,7 +37,6 @@ void FChaosBreakingEventFilter::FilterEvents(const FTransform& ChaosComponentTra
 			NewData.Location = BreakingData.Location;
 			NewData.Velocity = BreakingData.Velocity;
 			NewData.Mass = BreakingData.Mass;
-			NewData.ParticleIndex = BreakingData.ParticleIndex;
 
 			FilteredDataArray.Add(NewData);
 
diff --git a/Engine/Source/Runtime/Experimental/GeometryCollectionEngine/Public/ChaosBreakingEventFilter.h b/Engine/Source/Runtime/Experimental/GeometryCollectionEngine/Public/ChaosBreakingEventFilter.h
--- a/Engine/Source/Runtime/Experimental/GeometryCollectionEngine/Public/ChaosBreakingEventFilter.h
+++ b/Engine/Source/Runtime/Experimental/GeometryCollectionEngine/Public/ChaosBreakingEventFilter.h
@@ -27,15 +27,10 @@ struct FChaosBreakingEventData
 	UPROPERTY(EditAnywhere, BlueprintReadOnly, Category = "Data")
 	float Mass;
 
-	// The particle index of the breaking event
-	UPROPERTY(EditAnywhere, BlueprintReadOnly, Category = "Data")
-	int32 ParticleIndex;
-
 	FChaosBreakingEventData()
 		: Location(FVector::ZeroVector)
 		, Velocity(FVector::ZeroVector)
 		, Mass(0.0f)
-		, ParticleIndex(INDEX_NONE)
 	{
 	}
 };
