Fix CIS

diff --git a/Engine/Plugins/Experimental/ChaosNiagara/Source/ChaosNiagara/Private/NiagaraDataInterfaceChaosDestruction.cpp b/Engine/Plugins/Experimental/ChaosNiagara/Source/ChaosNiagara/Private/NiagaraDataInterfaceChaosDestruction.cpp
--- a/Engine/Plugins/Experimental/ChaosNiagara/Source/ChaosNiagara/Private/NiagaraDataInterfaceChaosDestruction.cpp
+++ b/Engine/Plugins/Experimental/ChaosNiagara/Source/ChaosNiagara/Private/NiagaraDataInterfaceChaosDestruction.cpp
@@ -1228,10 +1228,13 @@ void UNiagaraDataInterfaceChaosDestruction::HandleBreakingEvents(const Chaos::FB
 
 			UPhysicalMaterial* PhysicalMaterial = nullptr;
 			int32 MaterialID = DataIn.Particle->Geometry()->GetMaterialIndex(0);
-			UGeometryCollectionComponent* GeometryCollectionComponent = Solvers[0].PhysScene->GetScene().GetOwningComponent<UGeometryCollectionComponent>(DataIn.ParticleProxy);
 
-			UMaterialInterface* Material = GeometryCollectionComponent->GetMaterial(MaterialID);
-			ensure(Material);
+			UGeometryCollectionComponent* GeometryCollectionComponent = nullptr;
+			UMaterialInterface* Material = nullptr;
+#if INCLUDE_CHAOS
+			GeometryCollectionComponent = Solvers[0].PhysScene->GetScene().GetOwningComponent<UGeometryCollectionComponent>(DataIn.ParticleProxy);
+			Material = GeometryCollectionComponent->GetMaterial(MaterialID);
+#endif
 			if (Material)
 			{
 				PhysicalMaterial = Material->GetPhysicalMaterial();
