Whitelist Linux for ChaosNiagara.uplugin

diff --git a/Engine/Plugins/Experimental/ChaosNiagara/ChaosNiagara.uplugin b/Engine/Plugins/Experimental/ChaosNiagara/ChaosNiagara.uplugin
--- a/Engine/Plugins/Experimental/ChaosNiagara/ChaosNiagara.uplugin
+++ b/Engine/Plugins/Experimental/ChaosNiagara/ChaosNiagara.uplugin
@@ -19,7 +19,7 @@
 			"Name": "ChaosNiagara",
 			"Type": "Runtime",
 			"LoadingPhase": "Default",
-			"WhitelistPlatforms": [ "Win64", "Win32", "PS4", "XboxOne", "Mac" ]
+			"WhitelistPlatforms": [ "Win64", "Win32", "PS4", "XboxOne", "Mac", "Linux" ]
 		}
 	],
 	"Plugins": [
@@ -32,4 +32,4 @@
 			"Enabled": true
 		}
 	]
-}
\ No newline at end of file
+}
