NiagaraDataInterfaceChaosDestruction
-- guard against dereferencing potentially null pointers.

diff --git a/Engine/Plugins/Experimental/ChaosNiagara/Source/ChaosNiagara/Private/NiagaraDataInterfaceChaosDestruction.cpp b/Engine/Plugins/Experimental/ChaosNiagara/Source/ChaosNiagara/Private/NiagaraDataInterfaceChaosDestruction.cpp
--- a/Engine/Plugins/Experimental/ChaosNiagara/Source/ChaosNiagara/Private/NiagaraDataInterfaceChaosDestruction.cpp
+++ b/Engine/Plugins/Experimental/ChaosNiagara/Source/ChaosNiagara/Private/NiagaraDataInterfaceChaosDestruction.cpp
@@ -509,14 +509,16 @@ bool UNiagaraDataInterfaceChaosDestruction::InitPerInstanceData(void* PerInstanc
 
 	ResetInstData(InstData);
 
-	check(SystemInstance);
-	FNiagaraDataInterfaceProxyChaosDestruction* ThisProxy = GetProxyAs<FNiagaraDataInterfaceProxyChaosDestruction>();
-	ENQUEUE_RENDER_COMMAND(FNiagaraChaosDestructionDICreateRTInstance)(
-		[ThisProxy, InstanceID = SystemInstance->GetId()](FRHICommandList& CmdList)
+	if (SystemInstance)
 	{
-		ThisProxy->CreatePerInstanceData(InstanceID);
+		FNiagaraDataInterfaceProxyChaosDestruction* ThisProxy = GetProxyAs<FNiagaraDataInterfaceProxyChaosDestruction>();
+		ENQUEUE_RENDER_COMMAND(FNiagaraChaosDestructionDICreateRTInstance)(
+			[ThisProxy, InstanceID = SystemInstance->GetId()](FRHICommandList& CmdList)
+		{
+			ThisProxy->CreatePerInstanceData(InstanceID);
+		}
+		);
 	}
-	);
 
 	return true;
 }
@@ -528,13 +530,16 @@ void UNiagaraDataInterfaceChaosDestruction::DestroyPerInstanceData(void* PerInst
 
 	check(SystemInstance);
 	check(Proxy);
-	FNiagaraDataInterfaceProxyChaosDestruction* ThisProxy = GetProxyAs<FNiagaraDataInterfaceProxyChaosDestruction>();
-	ENQUEUE_RENDER_COMMAND(FNiagaraDIChaosDestructionDestroyInstanceData) (
-		[ThisProxy, InstanceID = SystemInstance->GetId(), Batcher = SystemInstance->GetBatcher()](FRHICommandListImmediate& CmdList)
+	if (SystemInstance)
+	{
+		FNiagaraDataInterfaceProxyChaosDestruction* ThisProxy = GetProxyAs<FNiagaraDataInterfaceProxyChaosDestruction>();
+		ENQUEUE_RENDER_COMMAND(FNiagaraDIChaosDestructionDestroyInstanceData) (
+			[ThisProxy, InstanceID = SystemInstance->GetId(), Batcher = SystemInstance->GetBatcher()](FRHICommandListImmediate& CmdList)
 		{
 			ThisProxy->DestroyInstanceData(Batcher, InstanceID);
 		}
-	);
+		);
+	}
 }
 
 #if CHAOS_PARTICLEHANDLE_TODO
