Fix GPU chaos field niagara DI + bad expresion caption

diff --git a/Engine/Plugins/Experimental/ChaosNiagara/Shaders/NiagaraDataInterfacePhysicsField.ush b/Engine/Plugins/Experimental/ChaosNiagara/Shaders/NiagaraDataInterfacePhysicsField.ush
--- a/Engine/Plugins/Experimental/ChaosNiagara/Shaders/NiagaraDataInterfacePhysicsField.ush
+++ b/Engine/Plugins/Experimental/ChaosNiagara/Shaders/NiagaraDataInterfacePhysicsField.ush
@@ -22,9 +22,7 @@ struct FDIPhysicsFieldContext
 	int ClipmapExponent;
 	int ClipmapCount;
 	int TargetCount;
-	int VectorTargets[MAX_PHYSICS_FIELD_TARGETS];
-	int ScalarTargets[MAX_PHYSICS_FIELD_TARGETS];
-	int IntegerTargets[MAX_PHYSICS_FIELD_TARGETS];
+	int4 FieldTargets[MAX_PHYSICS_FIELD_TARGETS];
 };
 
 #define DIPhysicsField_DECLARE_CONSTANTS(NAME)\
@@ -35,9 +33,7 @@ int ClipmapResolution_##NAME;\
 int ClipmapExponent_##NAME;\
 int ClipmapCount_##NAME;\
 int TargetCount_##NAME;\
-int VectorTargets_##NAME[MAX_PHYSICS_FIELD_TARGETS];\
-int ScalarTargets_##NAME[MAX_PHYSICS_FIELD_TARGETS];\
-int IntegerTargets_##NAME[MAX_PHYSICS_FIELD_TARGETS];\
+int4 FieldTargets_##NAME[MAX_PHYSICS_FIELD_TARGETS];\
 
 #define DIPhysicsField_MAKE_CONTEXT(NAME)\
 FDIPhysicsFieldContext DIContext; \
@@ -48,24 +44,22 @@ DIContext.ClipmapResolution = ClipmapResolution_##NAME;\
 DIContext.ClipmapExponent = ClipmapExponent_##NAME;\
 DIContext.ClipmapCount = ClipmapCount_##NAME;\
 DIContext.TargetCount = TargetCount_##NAME;\
-DIContext.VectorTargets = VectorTargets_##NAME;\
-DIContext.ScalarTargets = ScalarTargets_##NAME;\
-DIContext.IntegerTargets = IntegerTargets_##NAME;\
+DIContext.FieldTargets = FieldTargets_##NAME;\
 
 float3 DIPhysicsField_SamplePhysicsVectorField(in FDIPhysicsFieldContext DIContext, in float3 WorldPosition, in int VectorTarget)
 {
-	return PhysicsField_SamplePhysicsVectorField(WorldPosition, VectorTarget, DIContext.VectorTargets, DIContext.TargetCount, DIContext.ClipmapCenter, DIContext.ClipmapDistance, 
+	return PhysicsField_SamplePhysicsVectorField(WorldPosition, VectorTarget, DIContext.FieldTargets, DIContext.TargetCount, DIContext.ClipmapCenter, DIContext.ClipmapDistance, 
 				DIContext.ClipmapExponent, DIContext.ClipmapCount, DIContext.ClipmapResolution, DIContext.ClipmapBuffer);
 }
 
 float DIPhysicsField_SamplePhysicsScalarField(in FDIPhysicsFieldContext DIContext, in float3 WorldPosition, in int ScalarTarget)
 {
-	return PhysicsField_SamplePhysicsScalarField(WorldPosition, ScalarTarget, DIContext.ScalarTargets, DIContext.TargetCount, DIContext.ClipmapCenter, DIContext.ClipmapDistance,
+	return PhysicsField_SamplePhysicsScalarField(WorldPosition, ScalarTarget, DIContext.FieldTargets, DIContext.TargetCount, DIContext.ClipmapCenter, DIContext.ClipmapDistance,
 				DIContext.ClipmapExponent, DIContext.ClipmapCount, DIContext.ClipmapResolution, DIContext.ClipmapBuffer);
 }
 
 int DIPhysicsField_SamplePhysicsIntegerField(in FDIPhysicsFieldContext DIContext, in float3 WorldPosition, in int IntegerTarget)
 {
-	return PhysicsField_SamplePhysicsIntegerField(WorldPosition, IntegerTarget, DIContext.IntegerTargets, DIContext.TargetCount, DIContext.ClipmapCenter, DIContext.ClipmapDistance,
+	return PhysicsField_SamplePhysicsIntegerField(WorldPosition, IntegerTarget, DIContext.FieldTargets, DIContext.TargetCount, DIContext.ClipmapCenter, DIContext.ClipmapDistance,
 				DIContext.ClipmapExponent, DIContext.ClipmapCount, DIContext.ClipmapResolution, DIContext.ClipmapBuffer);
 }
\ No newline at end of file
diff --git a/Engine/Plugins/Experimental/ChaosNiagara/Source/ChaosNiagara/Classes/NiagaraDataInterfacePhysicsField.h b/Engine/Plugins/Experimental/ChaosNiagara/Source/ChaosNiagara/Classes/NiagaraDataInterfacePhysicsField.h
--- a/Engine/Plugins/Experimental/ChaosNiagara/Source/ChaosNiagara/Classes/NiagaraDataInterfacePhysicsField.h
+++ b/Engine/Plugins/Experimental/ChaosNiagara/Source/ChaosNiagara/Classes/NiagaraDataInterfacePhysicsField.h
@@ -86,9 +86,7 @@ public:
 	static const FString ClipmapExponentName;
 	static const FString ClipmapCountName;
 	static const FString TargetCountName;
-	static const FString VectorTargetsName;
-	static const FString ScalarTargetsName;
-	static const FString IntegerTargetsName;
+	static const FString FieldTargetsName;
 
 protected:
 	/** Copy one niagara DI to this */
diff --git a/Engine/Plugins/Experimental/ChaosNiagara/Source/ChaosNiagara/Private/NiagaraDataInterfacePhysicsField.cpp b/Engine/Plugins/Experimental/ChaosNiagara/Source/ChaosNiagara/Private/NiagaraDataInterfacePhysicsField.cpp
--- a/Engine/Plugins/Experimental/ChaosNiagara/Source/ChaosNiagara/Private/NiagaraDataInterfacePhysicsField.cpp
+++ b/Engine/Plugins/Experimental/ChaosNiagara/Source/ChaosNiagara/Private/NiagaraDataInterfacePhysicsField.cpp
@@ -32,9 +32,7 @@ const FString UNiagaraDataInterfacePhysicsField::ClipmapResolutionName(TEXT("Cli
 const FString UNiagaraDataInterfacePhysicsField::ClipmapExponentName(TEXT("ClipmapExponent_"));
 const FString UNiagaraDataInterfacePhysicsField::ClipmapCountName(TEXT("ClipmapCount_"));
 const FString UNiagaraDataInterfacePhysicsField::TargetCountName(TEXT("TargetCount_"));
-const FString UNiagaraDataInterfacePhysicsField::VectorTargetsName(TEXT("VectorTargets_"));
-const FString UNiagaraDataInterfacePhysicsField::ScalarTargetsName(TEXT("ScalarTargets_"));
-const FString UNiagaraDataInterfacePhysicsField::IntegerTargetsName(TEXT("IntegerTargets_"));
+const FString UNiagaraDataInterfacePhysicsField::FieldTargetsName(TEXT("FieldTargets_"));
 
 //------------------------------------------------------------------------------------------------------------
 
@@ -49,9 +47,7 @@ struct FNDIPhysicsFieldParametersName
 		ClipmapExponentName = UNiagaraDataInterfacePhysicsField::ClipmapExponentName + Suffix;
 		ClipmapCountName = UNiagaraDataInterfacePhysicsField::ClipmapCountName + Suffix;
 		TargetCountName = UNiagaraDataInterfacePhysicsField::TargetCountName + Suffix;
-		VectorTargetsName = UNiagaraDataInterfacePhysicsField::VectorTargetsName + Suffix;
-		ScalarTargetsName = UNiagaraDataInterfacePhysicsField::ScalarTargetsName + Suffix;
-		IntegerTargetsName = UNiagaraDataInterfacePhysicsField::IntegerTargetsName + Suffix;
+		FieldTargetsName = UNiagaraDataInterfacePhysicsField::FieldTargetsName + Suffix;
 	}
 
 	FString ClipmapBufferName;
@@ -61,10 +57,7 @@ struct FNDIPhysicsFieldParametersName
 	FString ClipmapExponentName;
 	FString ClipmapCountName;
 	FString TargetCountName;
-	FString VectorTargetsName;
-	FString ScalarTargetsName;
-	FString IntegerTargetsName;
-
+	FString FieldTargetsName;
 };
 
 //------------------------------------------------------------------------------------------------------------
@@ -129,9 +122,7 @@ public:
 		ClipmapExponent.Bind(ParameterMap, *ParamNames.ClipmapExponentName);
 		ClipmapCount.Bind(ParameterMap, *ParamNames.ClipmapCountName);
 		TargetCount.Bind(ParameterMap, *ParamNames.TargetCountName);
-		VectorTargets.Bind(ParameterMap, *ParamNames.VectorTargetsName);
-		ScalarTargets.Bind(ParameterMap, *ParamNames.ScalarTargetsName);
-		IntegerTargets.Bind(ParameterMap, *ParamNames.IntegerTargetsName);
+		FieldTargets.Bind(ParameterMap, *ParamNames.FieldTargetsName);
 	}
 
 	void Set(FRHICommandList& RHICmdList, const FNiagaraDataInterfaceSetArgs& Context) const
@@ -145,10 +136,18 @@ public:
 		FNDIPhysicsFieldData* ProxyData =
 			InterfaceProxy->SystemInstancesToProxyData.Find(Context.SystemInstanceID);
 
+		TStaticArray<FIntVector4, MAX_PHYSICS_FIELD_TARGETS,16> LocalTargets;
 		if (ProxyData != nullptr && ProxyData->FieldResource)
 		{
 			FPhysicsFieldResource* FieldResource = ProxyData->FieldResource;
 
+			for (int32 Index = 0; Index < MAX_PHYSICS_FIELD_TARGETS; ++Index)
+			{
+				LocalTargets[Index].X = FieldResource->FieldInfos.VectorTargets[Index];
+				LocalTargets[Index].Y = FieldResource->FieldInfos.ScalarTargets[Index];
+				LocalTargets[Index].Z = FieldResource->FieldInfos.IntegerTargets[Index];
+				LocalTargets[Index].W = 0; // Padding
+			}
 			SetSRVParameter(RHICmdList, ComputeShaderRHI, ClipmapBuffer, FieldResource->ClipmapBuffer.SRV);
 			SetShaderValue(RHICmdList, ComputeShaderRHI, ClipmapCenter, FieldResource->FieldInfos.ClipmapCenter);
 			SetShaderValue(RHICmdList, ComputeShaderRHI, ClipmapDistance, FieldResource->FieldInfos.ClipmapDistance);
@@ -157,13 +156,10 @@ public:
 			SetShaderValue(RHICmdList, ComputeShaderRHI, ClipmapCount, FieldResource->FieldInfos.ClipmapCount);
 
 			SetShaderValue(RHICmdList, ComputeShaderRHI, TargetCount, FieldResource->FieldInfos.TargetCount);
-			SetShaderValue(RHICmdList, ComputeShaderRHI, VectorTargets, FieldResource->FieldInfos.VectorTargets);
-			SetShaderValue(RHICmdList, ComputeShaderRHI, ScalarTargets, FieldResource->FieldInfos.ScalarTargets);
-			SetShaderValue(RHICmdList, ComputeShaderRHI, IntegerTargets, FieldResource->FieldInfos.IntegerTargets);
+			SetShaderValue(RHICmdList, ComputeShaderRHI, FieldTargets, LocalTargets);
 		}
 		else
 		{
-			TStaticArray<int32, MAX_PHYSICS_FIELD_TARGETS> EmptyTargets;
 			SetSRVParameter(RHICmdList, ComputeShaderRHI, ClipmapBuffer, FNiagaraRenderer::GetDummyFloatBuffer());
 			SetShaderValue(RHICmdList, ComputeShaderRHI, ClipmapCenter, FVector::ZeroVector);
 			SetShaderValue(RHICmdList, ComputeShaderRHI, ClipmapDistance, 1);
@@ -171,9 +167,7 @@ public:
 			SetShaderValue(RHICmdList, ComputeShaderRHI, ClipmapExponent, 1);
 			SetShaderValue(RHICmdList, ComputeShaderRHI, ClipmapCount, 1);
 			SetShaderValue(RHICmdList, ComputeShaderRHI, TargetCount, 0);
-			SetShaderValue(RHICmdList, ComputeShaderRHI, VectorTargets, EmptyTargets);
-			SetShaderValue(RHICmdList, ComputeShaderRHI, ScalarTargets, EmptyTargets);
-			SetShaderValue(RHICmdList, ComputeShaderRHI, IntegerTargets, EmptyTargets);
+			SetShaderValue(RHICmdList, ComputeShaderRHI, FieldTargets, LocalTargets);
 		}
 	}
 
@@ -190,9 +184,7 @@ private:
 	LAYOUT_FIELD(FShaderParameter, ClipmapExponent);
 	LAYOUT_FIELD(FShaderParameter, ClipmapCount);
 	LAYOUT_FIELD(FShaderParameter, TargetCount);
-	LAYOUT_FIELD(FShaderParameter, VectorTargets);
-	LAYOUT_FIELD(FShaderParameter, ScalarTargets);
-	LAYOUT_FIELD(FShaderParameter, IntegerTargets);
+	LAYOUT_FIELD(FShaderParameter, FieldTargets);
 };
 
 IMPLEMENT_TYPE_LAYOUT(FNDIPhysicsFieldParametersCS);
diff --git a/Engine/Source/Runtime/Engine/Private/Materials/MaterialExpressions.cpp b/Engine/Source/Runtime/Engine/Private/Materials/MaterialExpressions.cpp
--- a/Engine/Source/Runtime/Engine/Private/Materials/MaterialExpressions.cpp
+++ b/Engine/Source/Runtime/Engine/Private/Materials/MaterialExpressions.cpp
@@ -16811,7 +16811,7 @@ int32 UMaterialExpressionSamplePhysicsIntegerField::Compile(class FMaterialCompi
 
 void UMaterialExpressionSamplePhysicsIntegerField::GetCaption(TArray<FString>& OutCaptions) const
 {
-	OutCaptions.Add(TEXT("SamplePhysicsScalarField"));
+	OutCaptions.Add(TEXT("SamplePhysicsIntegerField"));
 }
 #endif // WITH_EDITOR
 
