Wrapped FVectorVMContext in NiagaraDataInterfaceFunctionContext for external functions to facilitate running an experimental VM in parallel.
External functions should no longer access FVectorVMContext directly and instead use the helper member functions in the NiagaraDataInterfaceFunctionContext class... this should have zero effect without NIAGARA_EXP_VM defined... if it is defined then it'll use the experimental VM functions, which are not included in this CL.

diff --git a/Engine/Plugins/Experimental/ChaosNiagara/Source/ChaosNiagara/Classes/NiagaraDataInterfaceChaosDestruction.h b/Engine/Plugins/Experimental/ChaosNiagara/Source/ChaosNiagara/Classes/NiagaraDataInterfaceChaosDestruction.h
--- a/Engine/Plugins/Experimental/ChaosNiagara/Source/ChaosNiagara/Classes/NiagaraDataInterfaceChaosDestruction.h
+++ b/Engine/Plugins/Experimental/ChaosNiagara/Source/ChaosNiagara/Classes/NiagaraDataInterfaceChaosDestruction.h
@@ -566,64 +566,64 @@ public:
 	//----------------------------------------------------------------------------
 	// EXPOSED FUNCTIONS
 	template<typename ParticleIDParamType>
-	void GetPosition(FVectorVMContext& Context);
+	void GetPosition(FVectorVMExternalFunctionContext& Context);
 
 	template<typename ParticleIDParamType>
-	void GetNormal(FVectorVMContext& Context);
+	void GetNormal(FVectorVMExternalFunctionContext& Context);
 
 	template<typename ParticleIDParamType>
-	void GetVelocity(FVectorVMContext& Context);
+	void GetVelocity(FVectorVMExternalFunctionContext& Context);
 
 	template<typename ParticleIDParamType>
-	void GetAngularVelocity(FVectorVMContext& Context);
+	void GetAngularVelocity(FVectorVMExternalFunctionContext& Context);
 
 	template<typename ParticleIDParamType>
-	void GetExtentMin(FVectorVMContext& Context);
+	void GetExtentMin(FVectorVMExternalFunctionContext& Context);
 
 	template<typename ParticleIDParamType>
-	void GetExtentMax(FVectorVMContext& Context);
+	void GetExtentMax(FVectorVMExternalFunctionContext& Context);
 
 	template<typename ParticleIDParamType>
-	void GetVolume(FVectorVMContext& Context);
+	void GetVolume(FVectorVMExternalFunctionContext& Context);
 
 	template<typename TimeParamType>
-	void GetParticleIdsToSpawnAtTime(FVectorVMContext& Context);
+	void GetParticleIdsToSpawnAtTime(FVectorVMExternalFunctionContext& Context);
 
 	template<typename ParticleIDParamType>
-	void GetPointType(FVectorVMContext& Context);
+	void GetPointType(FVectorVMExternalFunctionContext& Context);
 
 	template<typename ParticleIDParamType>
-	void GetColor(FVectorVMContext& Context);
+	void GetColor(FVectorVMExternalFunctionContext& Context);
 
 	template<typename ParticleIDParamType>
-	void GetSolverTime(FVectorVMContext& Context);
+	void GetSolverTime(FVectorVMExternalFunctionContext& Context);
 
 	template<typename ParticleIDParamType>
-	void GetDensity(FVectorVMContext& Context);
+	void GetDensity(FVectorVMExternalFunctionContext& Context);
 
 	template<typename ParticleIDParamType>
-	void GetFriction(FVectorVMContext& Context);
+	void GetFriction(FVectorVMExternalFunctionContext& Context);
 
 	template<typename ParticleIDParamType>
-	void GetRestitution(FVectorVMContext& Context);
+	void GetRestitution(FVectorVMExternalFunctionContext& Context);
 
 	template<typename ParticleIDParamType>
-	void GetSize(FVectorVMContext& Context);
+	void GetSize(FVectorVMExternalFunctionContext& Context);
 
 	template<typename ParticleIDParamType>
-	void GetTransform(FVectorVMContext& Context);
+	void GetTransform(FVectorVMExternalFunctionContext& Context);
 
 	template<typename ParticleIDParamType>
-	void GetSurfaceType(FVectorVMContext& Context);
+	void GetSurfaceType(FVectorVMExternalFunctionContext& Context);
 
 	template<typename ParticleIDParamType>
-	void GetCollisionData(FVectorVMContext& Context);
+	void GetCollisionData(FVectorVMExternalFunctionContext& Context);
 
 	template<typename ParticleIDParamType>
-	void GetBreakingData(FVectorVMContext& Context);
+	void GetBreakingData(FVectorVMExternalFunctionContext& Context);
 
 	template<typename ParticleIDParamType>
-	void GetTrailingData(FVectorVMContext& Context);
+	void GetTrailingData(FVectorVMExternalFunctionContext& Context);
 
 	// Sort predicates to sort data
 	inline static bool CollisionDataSortByMassPredicateMaxToMin(const Chaos::FCollidingDataExt& Lhs, const Chaos::FCollidingDataExt& Rhs)
diff --git a/Engine/Plugins/Experimental/ChaosNiagara/Source/ChaosNiagara/Classes/NiagaraDataInterfacePhysicsField.h b/Engine/Plugins/Experimental/ChaosNiagara/Source/ChaosNiagara/Classes/NiagaraDataInterfacePhysicsField.h
--- a/Engine/Plugins/Experimental/ChaosNiagara/Source/ChaosNiagara/Classes/NiagaraDataInterfacePhysicsField.h
+++ b/Engine/Plugins/Experimental/ChaosNiagara/Source/ChaosNiagara/Classes/NiagaraDataInterfacePhysicsField.h
@@ -66,19 +66,19 @@ public:
 	virtual void ProvidePerInstanceDataForRenderThread(void* DataForRenderThread, void* PerInstanceData, const FNiagaraSystemInstanceID& SystemInstance) override;
 
 	/** Sample the vector field */
-	void SamplePhysicsVectorField(FVectorVMContext& Context);
+	void SamplePhysicsVectorField(FVectorVMExternalFunctionContext& Context);
 
 	/** Sample the scalar field */
-	void SamplePhysicsScalarField(FVectorVMContext& Context);
+	void SamplePhysicsScalarField(FVectorVMExternalFunctionContext& Context);
 
 	/** Sample the integer field */
-	void SamplePhysicsIntegerField(FVectorVMContext& Context);
+	void SamplePhysicsIntegerField(FVectorVMExternalFunctionContext& Context);
 
 	/** Get the field resolution */
-	void GetPhysicsFieldResolution(FVectorVMContext& Context);
+	void GetPhysicsFieldResolution(FVectorVMExternalFunctionContext& Context);
 
 	/** Get the field bounds */
-	void GetPhysicsFieldBounds(FVectorVMContext& Context);
+	void GetPhysicsFieldBounds(FVectorVMExternalFunctionContext& Context);
 
 	/** Shader attributes names */
 	static const FString ClipmapBufferName;
diff --git a/Engine/Plugins/Experimental/ChaosNiagara/Source/ChaosNiagara/Private/NiagaraDataInterfaceChaosDestruction.cpp b/Engine/Plugins/Experimental/ChaosNiagara/Source/ChaosNiagara/Private/NiagaraDataInterfaceChaosDestruction.cpp
--- a/Engine/Plugins/Experimental/ChaosNiagara/Source/ChaosNiagara/Private/NiagaraDataInterfaceChaosDestruction.cpp
+++ b/Engine/Plugins/Experimental/ChaosNiagara/Source/ChaosNiagara/Private/NiagaraDataInterfaceChaosDestruction.cpp
@@ -2687,7 +2687,7 @@ void UNiagaraDataInterfaceChaosDestruction::GetVMExternalFunction(const FVMExter
 }
 
 template<typename ParticleIDParamType>
-void UNiagaraDataInterfaceChaosDestruction::GetPosition(FVectorVMContext& Context)
+void UNiagaraDataInterfaceChaosDestruction::GetPosition(FVectorVMExternalFunctionContext& Context)
 {
 	VectorVM::FUserPtrHandler<FNDIChaosDestruction_InstanceData> InstData(Context);
 	ParticleIDParamType ParticleIDParam(Context);
@@ -2695,7 +2695,7 @@ void UNiagaraDataInterfaceChaosDestruction::GetPosition(FVectorVMContext& Contex
 	VectorVM::FExternalFuncRegisterHandler<float> OutSampleY(Context);
 	VectorVM::FExternalFuncRegisterHandler<float> OutSampleZ(Context);
 
-	for (int32 i = 0; i < Context.NumInstances; ++i)
+	for (int32 i = 0; i < Context.GetNumInstances(); ++i)
 	{
 		//		if (InstData->ParticleDataArray.Num())
 		if (InstData->PositionArray.Num())
@@ -2718,7 +2718,7 @@ void UNiagaraDataInterfaceChaosDestruction::GetPosition(FVectorVMContext& Contex
 }
 
 template<typename ParticleIDParamType>
-void UNiagaraDataInterfaceChaosDestruction::GetNormal(FVectorVMContext& Context)
+void UNiagaraDataInterfaceChaosDestruction::GetNormal(FVectorVMExternalFunctionContext& Context)
 {
 	VectorVM::FUserPtrHandler<FNDIChaosDestruction_InstanceData> InstData(Context);
 	ParticleIDParamType ParticleIDParam(Context);
@@ -2726,7 +2726,7 @@ void UNiagaraDataInterfaceChaosDestruction::GetNormal(FVectorVMContext& Context)
 	VectorVM::FExternalFuncRegisterHandler<float> OutSampleY(Context);
 	VectorVM::FExternalFuncRegisterHandler<float> OutSampleZ(Context);
 
-	for (int32 i = 0; i < Context.NumInstances; ++i)
+	for (int32 i = 0; i < Context.GetNumInstances(); ++i)
 	{
 		if (InstData->IncomingNormalArray.Num())
 		{
@@ -2748,7 +2748,7 @@ void UNiagaraDataInterfaceChaosDestruction::GetNormal(FVectorVMContext& Context)
 }
 
 template<typename ParticleIDParamType>
-void UNiagaraDataInterfaceChaosDestruction::GetVelocity(FVectorVMContext& Context)
+void UNiagaraDataInterfaceChaosDestruction::GetVelocity(FVectorVMExternalFunctionContext& Context)
 {
 	VectorVM::FUserPtrHandler<FNDIChaosDestruction_InstanceData> InstData(Context);
 	ParticleIDParamType ParticleIDParam(Context);
@@ -2756,7 +2756,7 @@ void UNiagaraDataInterfaceChaosDestruction::GetVelocity(FVectorVMContext& Contex
 	VectorVM::FExternalFuncRegisterHandler<float> OutSampleY(Context);
 	VectorVM::FExternalFuncRegisterHandler<float> OutSampleZ(Context);
 
-	for (int32 i = 0; i < Context.NumInstances; ++i)
+	for (int32 i = 0; i < Context.GetNumInstances(); ++i)
 	{
 		//		if (InstData->ParticleDataArray.Num())
 		if (InstData->VelocityArray.Num())
@@ -2778,7 +2778,7 @@ void UNiagaraDataInterfaceChaosDestruction::GetVelocity(FVectorVMContext& Contex
 }
 
 template<typename ParticleIDParamType>
-void UNiagaraDataInterfaceChaosDestruction::GetAngularVelocity(FVectorVMContext& Context)
+void UNiagaraDataInterfaceChaosDestruction::GetAngularVelocity(FVectorVMExternalFunctionContext& Context)
 {
 	VectorVM::FUserPtrHandler<FNDIChaosDestruction_InstanceData> InstData(Context);
 	ParticleIDParamType ParticleIDParam(Context);
@@ -2786,7 +2786,7 @@ void UNiagaraDataInterfaceChaosDestruction::GetAngularVelocity(FVectorVMContext&
 	VectorVM::FExternalFuncRegisterHandler<float> OutSampleY(Context);
 	VectorVM::FExternalFuncRegisterHandler<float> OutSampleZ(Context);
 
-	for (int32 i = 0; i < Context.NumInstances; ++i)
+	for (int32 i = 0; i < Context.GetNumInstances(); ++i)
 	{
 		if (InstData->IncomingAngularVelocity1Array.Num())
 		{
@@ -2808,13 +2808,13 @@ void UNiagaraDataInterfaceChaosDestruction::GetAngularVelocity(FVectorVMContext&
 }
 
 template<typename ParticleIDParamType>
-void UNiagaraDataInterfaceChaosDestruction::GetExtentMin(FVectorVMContext& Context)
+void UNiagaraDataInterfaceChaosDestruction::GetExtentMin(FVectorVMExternalFunctionContext& Context)
 {
 	VectorVM::FUserPtrHandler<FNDIChaosDestruction_InstanceData> InstData(Context);
 	ParticleIDParamType ParticleIDParam(Context);
 	VectorVM::FExternalFuncRegisterHandler<float> OutValue(Context);
 
-	for (int32 i = 0; i < Context.NumInstances; ++i)
+	for (int32 i = 0; i < Context.GetNumInstances(); ++i)
 	{
 		if (InstData->ExtentMinArray.Num())
 		{
@@ -2832,13 +2832,13 @@ void UNiagaraDataInterfaceChaosDestruction::GetExtentMin(FVectorVMContext& Conte
 }
 
 template<typename ParticleIDParamType>
-void UNiagaraDataInterfaceChaosDestruction::GetExtentMax(FVectorVMContext& Context)
+void UNiagaraDataInterfaceChaosDestruction::GetExtentMax(FVectorVMExternalFunctionContext& Context)
 {
 	VectorVM::FUserPtrHandler<FNDIChaosDestruction_InstanceData> InstData(Context);
 	ParticleIDParamType ParticleIDParam(Context);
 	VectorVM::FExternalFuncRegisterHandler<float> OutValue(Context);
 
-	for (int32 i = 0; i < Context.NumInstances; ++i)
+	for (int32 i = 0; i < Context.GetNumInstances(); ++i)
 	{
 		if (InstData->ExtentMaxArray.Num())
 		{
@@ -2856,13 +2856,13 @@ void UNiagaraDataInterfaceChaosDestruction::GetExtentMax(FVectorVMContext& Conte
 }
 
 template<typename ParticleIDParamType>
-void UNiagaraDataInterfaceChaosDestruction::GetVolume(FVectorVMContext& Context)
+void UNiagaraDataInterfaceChaosDestruction::GetVolume(FVectorVMExternalFunctionContext& Context)
 {
 	VectorVM::FUserPtrHandler<FNDIChaosDestruction_InstanceData> InstData(Context);
 	ParticleIDParamType ParticleIDParam(Context);
 	VectorVM::FExternalFuncRegisterHandler<float> OutValue(Context);
 
-	for (int32 i = 0; i < Context.NumInstances; ++i)
+	for (int32 i = 0; i < Context.GetNumInstances(); ++i)
 	{
 		if (InstData->VolumeArray.Num())
 		{
@@ -2880,7 +2880,7 @@ void UNiagaraDataInterfaceChaosDestruction::GetVolume(FVectorVMContext& Context)
 }
 
 template<typename TimeParamType>
-void UNiagaraDataInterfaceChaosDestruction::GetParticleIdsToSpawnAtTime(FVectorVMContext& Context)
+void UNiagaraDataInterfaceChaosDestruction::GetParticleIdsToSpawnAtTime(FVectorVMExternalFunctionContext& Context)
 {
 	VectorVM::FUserPtrHandler<FNDIChaosDestruction_InstanceData> InstData(Context);
 	TimeParamType TimeParam(Context);
@@ -2910,13 +2910,13 @@ void UNiagaraDataInterfaceChaosDestruction::GetParticleIdsToSpawnAtTime(FVectorV
 }
 
 template<typename ParticleIDParamType>
-void UNiagaraDataInterfaceChaosDestruction::GetPointType(FVectorVMContext& Context)
+void UNiagaraDataInterfaceChaosDestruction::GetPointType(FVectorVMExternalFunctionContext& Context)
 {
 	VectorVM::FUserPtrHandler<FNDIChaosDestruction_InstanceData> InstData(Context);
 	ParticleIDParamType ParticleIDParam(Context);
 	VectorVM::FExternalFuncRegisterHandler<int32> OutValue(Context);
 
-	for (int32 i = 0; i < Context.NumInstances; ++i)
+	for (int32 i = 0; i < Context.GetNumInstances(); ++i)
 	{
 		if (InstData->PositionArray.Num())
 		{
@@ -2935,7 +2935,7 @@ void UNiagaraDataInterfaceChaosDestruction::GetPointType(FVectorVMContext& Conte
 }
 
 template<typename ParticleIDParamType>
-void UNiagaraDataInterfaceChaosDestruction::GetColor(FVectorVMContext& Context)
+void UNiagaraDataInterfaceChaosDestruction::GetColor(FVectorVMExternalFunctionContext& Context)
 {
 	VectorVM::FUserPtrHandler<FNDIChaosDestruction_InstanceData> InstData(Context);
 	ParticleIDParamType ParticleIDParam(Context);
@@ -2944,7 +2944,7 @@ void UNiagaraDataInterfaceChaosDestruction::GetColor(FVectorVMContext& Context)
 	VectorVM::FExternalFuncRegisterHandler<float> OutSampleB(Context);
 	VectorVM::FExternalFuncRegisterHandler<float> OutSampleA(Context);
 
-	for (int32 i = 0; i < Context.NumInstances; ++i)
+	for (int32 i = 0; i < Context.GetNumInstances(); ++i)
 	{
 		if (InstData->ColorArray.Num())
 		{
@@ -2968,7 +2968,7 @@ void UNiagaraDataInterfaceChaosDestruction::GetColor(FVectorVMContext& Context)
 }
 
 template<typename ParticleIDParamType>
-void UNiagaraDataInterfaceChaosDestruction::GetSolverTime(FVectorVMContext& Context)
+void UNiagaraDataInterfaceChaosDestruction::GetSolverTime(FVectorVMExternalFunctionContext& Context)
 {
 	VectorVM::FUserPtrHandler<FNDIChaosDestruction_InstanceData> InstData(Context);
 	VectorVM::FExternalFuncRegisterHandler<float> OutValue(Context);
@@ -2979,13 +2979,13 @@ void UNiagaraDataInterfaceChaosDestruction::GetSolverTime(FVectorVMContext& Cont
 }
 
 template<typename ParticleIDParamType>
-void UNiagaraDataInterfaceChaosDestruction::GetDensity(FVectorVMContext& Context)
+void UNiagaraDataInterfaceChaosDestruction::GetDensity(FVectorVMExternalFunctionContext& Context)
 {
 	VectorVM::FUserPtrHandler<FNDIChaosDestruction_InstanceData> InstData(Context);
 	ParticleIDParamType ParticleIDParam(Context);
 	VectorVM::FExternalFuncRegisterHandler<float> OutValue(Context);
 
-	for (int32 i = 0; i < Context.NumInstances; ++i)
+	for (int32 i = 0; i < Context.GetNumInstances(); ++i)
 	{
 		if (InstData->DensityArray.Num())
 		{
@@ -3003,13 +3003,13 @@ void UNiagaraDataInterfaceChaosDestruction::GetDensity(FVectorVMContext& Context
 }
 
 template<typename ParticleIDParamType>
-void UNiagaraDataInterfaceChaosDestruction::GetFriction(FVectorVMContext& Context)
+void UNiagaraDataInterfaceChaosDestruction::GetFriction(FVectorVMExternalFunctionContext& Context)
 {
 	VectorVM::FUserPtrHandler<FNDIChaosDestruction_InstanceData> InstData(Context);
 	ParticleIDParamType ParticleIDParam(Context);
 	VectorVM::FExternalFuncRegisterHandler<float> OutValue(Context);
 
-	for (int32 i = 0; i < Context.NumInstances; ++i)
+	for (int32 i = 0; i < Context.GetNumInstances(); ++i)
 	{
 		if (InstData->FrictionArray.Num())
 		{
@@ -3027,13 +3027,13 @@ void UNiagaraDataInterfaceChaosDestruction::GetFriction(FVectorVMContext& Contex
 }
 
 template<typename ParticleIDParamType>
-void UNiagaraDataInterfaceChaosDestruction::GetRestitution(FVectorVMContext& Context)
+void UNiagaraDataInterfaceChaosDestruction::GetRestitution(FVectorVMExternalFunctionContext& Context)
 {
 	VectorVM::FUserPtrHandler<FNDIChaosDestruction_InstanceData> InstData(Context);
 	ParticleIDParamType ParticleIDParam(Context);
 	VectorVM::FExternalFuncRegisterHandler<float> OutValue(Context);
 
-	for (int32 i = 0; i < Context.NumInstances; ++i)
+	for (int32 i = 0; i < Context.GetNumInstances(); ++i)
 	{
 		if (InstData->RestitutionArray.Num())
 		{
@@ -3051,7 +3051,7 @@ void UNiagaraDataInterfaceChaosDestruction::GetRestitution(FVectorVMContext& Con
 }
 
 template<typename ParticleIDParamType>
-void UNiagaraDataInterfaceChaosDestruction::GetTransform(FVectorVMContext& Context)
+void UNiagaraDataInterfaceChaosDestruction::GetTransform(FVectorVMExternalFunctionContext& Context)
 {
 	VectorVM::FUserPtrHandler<FNDIChaosDestruction_InstanceData> InstData(Context);
 	ParticleIDParamType ParticleIDParam(Context);
@@ -3066,7 +3066,7 @@ void UNiagaraDataInterfaceChaosDestruction::GetTransform(FVectorVMContext& Conte
 	VectorVM::FExternalFuncRegisterHandler<float> OutScaleY(Context);
 	VectorVM::FExternalFuncRegisterHandler<float> OutScaleZ(Context);
 
-	for (int32 i = 0; i < Context.NumInstances; ++i)
+	for (int32 i = 0; i < Context.GetNumInstances(); ++i)
 	{
 		if (InstData->TransformTranslationArray.Num())
 		{
@@ -3108,7 +3108,7 @@ void UNiagaraDataInterfaceChaosDestruction::GetTransform(FVectorVMContext& Conte
 }
 
 template<typename ParticleIDParamType>
-void UNiagaraDataInterfaceChaosDestruction::GetSize(FVectorVMContext& Context)
+void UNiagaraDataInterfaceChaosDestruction::GetSize(FVectorVMExternalFunctionContext& Context)
 {
 	VectorVM::FUserPtrHandler<FNDIChaosDestruction_InstanceData> InstData(Context);
 	ParticleIDParamType ParticleIDParam(Context);
@@ -3116,7 +3116,7 @@ void UNiagaraDataInterfaceChaosDestruction::GetSize(FVectorVMContext& Context)
 	VectorVM::FExternalFuncRegisterHandler<float> OutSizeY(Context);
 	VectorVM::FExternalFuncRegisterHandler<float> OutSizeZ(Context);
 
-	for (int32 i = 0; i < Context.NumInstances; ++i)
+	for (int32 i = 0; i < Context.GetNumInstances(); ++i)
 	{
 		if (InstData->TransformTranslationArray.Num())
 		{
@@ -3138,13 +3138,13 @@ void UNiagaraDataInterfaceChaosDestruction::GetSize(FVectorVMContext& Context)
 }
 
 template<typename ParticleIDParamType>
-void UNiagaraDataInterfaceChaosDestruction::GetSurfaceType(FVectorVMContext& Context)
+void UNiagaraDataInterfaceChaosDestruction::GetSurfaceType(FVectorVMExternalFunctionContext& Context)
 {
 	VectorVM::FUserPtrHandler<FNDIChaosDestruction_InstanceData> InstData(Context);
 	ParticleIDParamType ParticleIDParam(Context);
 	VectorVM::FExternalFuncRegisterHandler<int32> OutValue(Context);
 
-	for (int32 i = 0; i < Context.NumInstances; ++i)
+	for (int32 i = 0; i < Context.GetNumInstances(); ++i)
 	{
 		if (InstData->SurfaceTypeArray.Num())
 		{
@@ -3162,7 +3162,7 @@ void UNiagaraDataInterfaceChaosDestruction::GetSurfaceType(FVectorVMContext& Con
 }
 
 template<typename ParticleIDParamType>
-void UNiagaraDataInterfaceChaosDestruction::GetCollisionData(FVectorVMContext& Context)
+void UNiagaraDataInterfaceChaosDestruction::GetCollisionData(FVectorVMExternalFunctionContext& Context)
 {
 	VectorVM::FUserPtrHandler<FNDIChaosDestruction_InstanceData> InstData(Context);
 	ParticleIDParamType ParticleIDParam(Context);
@@ -3191,7 +3191,7 @@ void UNiagaraDataInterfaceChaosDestruction::GetCollisionData(FVectorVMContext& C
 	VectorVM::FExternalFuncRegisterHandler<float> OutMass2(Context);
 	VectorVM::FExternalFuncRegisterHandler<float> OutTime(Context);
 
-	for (int32 i = 0; i < Context.NumInstances; ++i)
+	for (int32 i = 0; i < Context.GetNumInstances(); ++i)
 	{
 		if (InstData->IncomingLocationArray.Num())
 		{
@@ -3276,7 +3276,7 @@ void UNiagaraDataInterfaceChaosDestruction::GetCollisionData(FVectorVMContext& C
 }
 
 template<typename ParticleIDParamType>
-void UNiagaraDataInterfaceChaosDestruction::GetBreakingData(FVectorVMContext& Context)
+void UNiagaraDataInterfaceChaosDestruction::GetBreakingData(FVectorVMExternalFunctionContext& Context)
 {
 	VectorVM::FUserPtrHandler<FNDIChaosDestruction_InstanceData> InstData(Context);
 	ParticleIDParamType ParticleIDParam(Context);
@@ -3292,7 +3292,7 @@ void UNiagaraDataInterfaceChaosDestruction::GetBreakingData(FVectorVMContext& Co
 	VectorVM::FExternalFuncRegisterHandler<float> OutMass(Context);
 	VectorVM::FExternalFuncRegisterHandler<float> OutTime(Context);
 
-	for (int32 i = 0; i < Context.NumInstances; ++i)
+	for (int32 i = 0; i < Context.GetNumInstances(); ++i)
 	{
 		if (InstData->IncomingLocationArray.Num())
 		{
@@ -3341,7 +3341,7 @@ void UNiagaraDataInterfaceChaosDestruction::GetBreakingData(FVectorVMContext& Co
 }
 
 template<typename ParticleIDParamType>
-void UNiagaraDataInterfaceChaosDestruction::GetTrailingData(FVectorVMContext& Context)
+void UNiagaraDataInterfaceChaosDestruction::GetTrailingData(FVectorVMExternalFunctionContext& Context)
 {
 	VectorVM::FUserPtrHandler<FNDIChaosDestruction_InstanceData> InstData(Context);
 	ParticleIDParamType ParticleIDParam(Context);
@@ -3357,7 +3357,7 @@ void UNiagaraDataInterfaceChaosDestruction::GetTrailingData(FVectorVMContext& Co
 	VectorVM::FExternalFuncRegisterHandler<float> OutMass(Context);
 	VectorVM::FExternalFuncRegisterHandler<float> OutTime(Context);
 
-	for (int32 i = 0; i < Context.NumInstances; ++i)
+	for (int32 i = 0; i < Context.GetNumInstances(); ++i)
 	{
 		if (InstData->IncomingLocationArray.Num())
 		{
diff --git a/Engine/Plugins/Experimental/ChaosNiagara/Source/ChaosNiagara/Private/NiagaraDataInterfacePhysicsField.cpp b/Engine/Plugins/Experimental/ChaosNiagara/Source/ChaosNiagara/Private/NiagaraDataInterfacePhysicsField.cpp
--- a/Engine/Plugins/Experimental/ChaosNiagara/Source/ChaosNiagara/Private/NiagaraDataInterfacePhysicsField.cpp
+++ b/Engine/Plugins/Experimental/ChaosNiagara/Source/ChaosNiagara/Private/NiagaraDataInterfacePhysicsField.cpp
@@ -402,7 +402,7 @@ void UNiagaraDataInterfacePhysicsField::GetVMExternalFunction(const FVMExternalF
 	}
 }
 
-void UNiagaraDataInterfacePhysicsField::GetPhysicsFieldResolution(FVectorVMContext& Context)
+void UNiagaraDataInterfacePhysicsField::GetPhysicsFieldResolution(FVectorVMExternalFunctionContext& Context)
 {
 	VectorVM::FUserPtrHandler<FNDIPhysicsFieldData> InstData(Context);
 
@@ -413,7 +413,7 @@ void UNiagaraDataInterfacePhysicsField::GetPhysicsFieldResolution(FVectorVMConte
 	const FIntVector FieldDimension = (InstData && InstData->FieldResource) ? FIntVector(InstData->FieldResource->FieldInfos.ClipmapResolution, 
 		InstData->FieldResource->FieldInfos.ClipmapResolution, InstData->FieldResource->FieldInfos.ClipmapResolution) : FIntVector(1, 1, 1);
 
-	for (int32 i = 0; i < Context.NumInstances; ++i)
+	for (int32 i = 0; i < Context.GetNumInstances(); ++i)
 	{
 		*OutDimensionX.GetDest() = FieldDimension.X;
 		*OutDimensionY.GetDest() = FieldDimension.Y;
@@ -425,7 +425,7 @@ void UNiagaraDataInterfacePhysicsField::GetPhysicsFieldResolution(FVectorVMConte
 	}
 }
 
-void UNiagaraDataInterfacePhysicsField::GetPhysicsFieldBounds(FVectorVMContext& Context)
+void UNiagaraDataInterfacePhysicsField::GetPhysicsFieldBounds(FVectorVMExternalFunctionContext& Context)
 {
 	VectorVM::FUserPtrHandler<FNDIPhysicsFieldData> InstData(Context);
 
@@ -441,7 +441,7 @@ void UNiagaraDataInterfacePhysicsField::GetPhysicsFieldBounds(FVectorVMContext&
 	const FVector MaxBound = (InstData && InstData->FieldResource) ? InstData->FieldResource->FieldInfos.ClipmapCenter +
 		FVector(InstData->FieldResource->FieldInfos.ClipmapDistance) : FVector(0, 0, 0);
 
-	for (int32 i = 0; i < Context.NumInstances; ++i)
+	for (int32 i = 0; i < Context.GetNumInstances(); ++i)
 	{
 		*OutMinX.GetDest() = MinBound.X;
 		*OutMinY.GetDest() = MinBound.Y;
@@ -522,7 +522,7 @@ void EvaluateFieldNodes(TArray<FFieldSystemCommand>& FieldCommands, const EField
 	}
 }
 
-void UNiagaraDataInterfacePhysicsField::SamplePhysicsVectorField(FVectorVMContext& Context)
+void UNiagaraDataInterfacePhysicsField::SamplePhysicsVectorField(FVectorVMExternalFunctionContext& Context)
 {
 	VectorVM::FUserPtrHandler<FNDIPhysicsFieldData> InstData(Context);
 
@@ -536,22 +536,22 @@ void UNiagaraDataInterfacePhysicsField::SamplePhysicsVectorField(FVectorVMContex
 	if (InstData)
 	{
 		FFieldExecutionDatas ExecutionDatas;
-		ExecutionDatas.SamplePositions.Init(FVector(0, 0, 0), Context.NumInstances);
+		ExecutionDatas.SamplePositions.Init(FVector(0, 0, 0), Context.GetNumInstances());
 
 		EFieldVectorType VectorTarget = EFieldVectorType::Vector_TargetMax;
 
-		for (int32 InstanceIdx = 0; InstanceIdx < Context.NumInstances; ++InstanceIdx)
+		for (int32 InstanceIdx = 0; InstanceIdx < Context.GetNumInstances(); ++InstanceIdx)
 		{
 			ExecutionDatas.SamplePositions[InstanceIdx] = SamplePositionParam.GetAndAdvance();
 			VectorTarget = VectorTargetParam.GetAndAdvance();
 		}
-		FFieldContextIndex::ContiguousIndices(ExecutionDatas.SampleIndices, Context.NumInstances);
+		FFieldContextIndex::ContiguousIndices(ExecutionDatas.SampleIndices, Context.GetNumInstances());
 
 		TArray<FVector>& SampleResults = ExecutionDatas.VectorResults[(uint8)EFieldCommandResultType::FinalResult];
-		SampleResults.Init(FVector::ZeroVector, Context.NumInstances);
+		SampleResults.Init(FVector::ZeroVector, Context.GetNumInstances());
 
 		TArray<FVector> SampleMax;
-		SampleMax.Init(FVector::ZeroVector, Context.NumInstances);
+		SampleMax.Init(FVector::ZeroVector, Context.GetNumInstances());
 
 		FFieldContext FieldContext{
 			ExecutionDatas,
@@ -562,21 +562,21 @@ void UNiagaraDataInterfacePhysicsField::SamplePhysicsVectorField(FVectorVMContex
 		const EFieldPhysicsType PhysicsType = GetFieldTargetTypes(EFieldOutputType::Field_Output_Vector)[VectorTarget];
 		EvaluateFieldNodes<FVector, FVectorFieldOperator>(InstData->FieldCommands, PhysicsType, FieldContext, SampleResults, SampleMax);
 
-		for (int32 InstanceIdx = 0; InstanceIdx < Context.NumInstances; ++InstanceIdx)
+		for (int32 InstanceIdx = 0; InstanceIdx < Context.GetNumInstances(); ++InstanceIdx)
 		{
 			OutVectorFieldParam.SetAndAdvance(SampleMax[InstanceIdx]);
 		}
 	}
 	else
 	{
-		for (int32 InstanceIdx = 0; InstanceIdx < Context.NumInstances; ++InstanceIdx)
+		for (int32 InstanceIdx = 0; InstanceIdx < Context.GetNumInstances(); ++InstanceIdx)
 		{
 			OutVectorFieldParam.SetAndAdvance(FVector::ZeroVector);
 		}
 	}
 }
 
-void UNiagaraDataInterfacePhysicsField::SamplePhysicsIntegerField(FVectorVMContext& Context)
+void UNiagaraDataInterfacePhysicsField::SamplePhysicsIntegerField(FVectorVMExternalFunctionContext& Context)
 {
 	VectorVM::FUserPtrHandler<FNDIPhysicsFieldData> InstData(Context);
 
@@ -590,22 +590,22 @@ void UNiagaraDataInterfacePhysicsField::SamplePhysicsIntegerField(FVectorVMConte
 	if (InstData)
 	{
 		FFieldExecutionDatas ExecutionDatas;
-		ExecutionDatas.SamplePositions.Init(FVector(0, 0, 0), Context.NumInstances);
+		ExecutionDatas.SamplePositions.Init(FVector(0, 0, 0), Context.GetNumInstances());
 
 		EFieldIntegerType IntegerTarget = EFieldIntegerType::Integer_TargetMax;
 
-		for (int32 InstanceIdx = 0; InstanceIdx < Context.NumInstances; ++InstanceIdx)
+		for (int32 InstanceIdx = 0; InstanceIdx < Context.GetNumInstances(); ++InstanceIdx)
 		{
 			ExecutionDatas.SamplePositions[InstanceIdx] = SamplePositionParam.GetAndAdvance();
 			IntegerTarget = IntegerTargetParam.GetAndAdvance();
 		}
-		FFieldContextIndex::ContiguousIndices(ExecutionDatas.SampleIndices, Context.NumInstances);
+		FFieldContextIndex::ContiguousIndices(ExecutionDatas.SampleIndices, Context.GetNumInstances());
 
 		TArray<int32>& SampleResults = ExecutionDatas.IntegerResults[(uint8)EFieldCommandResultType::FinalResult];
-		SampleResults.Init(0, Context.NumInstances);
+		SampleResults.Init(0, Context.GetNumInstances());
 
 		TArray<int32> SampleMax;
-		SampleMax.Init(0, Context.NumInstances);
+		SampleMax.Init(0, Context.GetNumInstances());
 
 		FFieldContext FieldContext{
 			ExecutionDatas,
@@ -616,14 +616,14 @@ void UNiagaraDataInterfacePhysicsField::SamplePhysicsIntegerField(FVectorVMConte
 		const EFieldPhysicsType PhysicsType = GetFieldTargetTypes(EFieldOutputType::Field_Output_Integer)[IntegerTarget];
 		EvaluateFieldNodes<int32, FIntegerFieldOperator>(InstData->FieldCommands, PhysicsType, FieldContext, SampleResults, SampleMax);
 
-		for (int32 InstanceIdx = 0; InstanceIdx < Context.NumInstances; ++InstanceIdx)
+		for (int32 InstanceIdx = 0; InstanceIdx < Context.GetNumInstances(); ++InstanceIdx)
 		{
 			OutIntegerFieldParam.SetAndAdvance(SampleMax[InstanceIdx]);
 		}
 	}
 	else
 	{
-		for (int32 InstanceIdx = 0; InstanceIdx < Context.NumInstances; ++InstanceIdx)
+		for (int32 InstanceIdx = 0; InstanceIdx < Context.GetNumInstances(); ++InstanceIdx)
 		{
 			OutIntegerFieldParam.SetAndAdvance(0.0);
 		}
@@ -631,7 +631,7 @@ void UNiagaraDataInterfacePhysicsField::SamplePhysicsIntegerField(FVectorVMConte
 }
 
 
-void UNiagaraDataInterfacePhysicsField::SamplePhysicsScalarField(FVectorVMContext& Context)
+void UNiagaraDataInterfacePhysicsField::SamplePhysicsScalarField(FVectorVMExternalFunctionContext& Context)
 {
 	VectorVM::FUserPtrHandler<FNDIPhysicsFieldData> InstData(Context);
 
@@ -645,22 +645,22 @@ void UNiagaraDataInterfacePhysicsField::SamplePhysicsScalarField(FVectorVMContex
 	if (InstData)
 	{
 		FFieldExecutionDatas ExecutionDatas;
-		ExecutionDatas.SamplePositions.Init(FVector(0, 0, 0), Context.NumInstances);
+		ExecutionDatas.SamplePositions.Init(FVector(0, 0, 0), Context.GetNumInstances());
 
 		EFieldScalarType ScalarTarget = EFieldScalarType::Scalar_TargetMax;
 
-		for (int32 InstanceIdx = 0; InstanceIdx < Context.NumInstances; ++InstanceIdx)
+		for (int32 InstanceIdx = 0; InstanceIdx < Context.GetNumInstances(); ++InstanceIdx)
 		{
 			ExecutionDatas.SamplePositions[InstanceIdx] = SamplePositionParam.GetAndAdvance();
 			ScalarTarget = ScalarTargetParam.GetAndAdvance();
 		}
-		FFieldContextIndex::ContiguousIndices(ExecutionDatas.SampleIndices, Context.NumInstances);
+		FFieldContextIndex::ContiguousIndices(ExecutionDatas.SampleIndices, Context.GetNumInstances());
 
 		TArray<float>& SampleResults = ExecutionDatas.ScalarResults[(uint8)EFieldCommandResultType::FinalResult];
-		SampleResults.Init(0.0, Context.NumInstances);
+		SampleResults.Init(0.0, Context.GetNumInstances());
 
 		TArray<float> SampleMax;
-		SampleMax.Init(0, Context.NumInstances);
+		SampleMax.Init(0, Context.GetNumInstances());
 
 		FFieldContext FieldContext{
 			ExecutionDatas,
@@ -671,14 +671,14 @@ void UNiagaraDataInterfacePhysicsField::SamplePhysicsScalarField(FVectorVMContex
 		const EFieldPhysicsType PhysicsType = GetFieldTargetTypes(EFieldOutputType::Field_Output_Scalar)[ScalarTarget];
 		EvaluateFieldNodes<float, FScalarFieldOperator>(InstData->FieldCommands, PhysicsType, FieldContext, SampleResults, SampleMax);
 
-		for (int32 InstanceIdx = 0; InstanceIdx < Context.NumInstances; ++InstanceIdx)
+		for (int32 InstanceIdx = 0; InstanceIdx < Context.GetNumInstances(); ++InstanceIdx)
 		{
 			OutScalarFieldParam.SetAndAdvance(SampleMax[InstanceIdx]);
 		}
 	}
 	else
 	{
-		for (int32 InstanceIdx = 0; InstanceIdx < Context.NumInstances; ++InstanceIdx)
+		for (int32 InstanceIdx = 0; InstanceIdx < Context.GetNumInstances(); ++InstanceIdx)
 		{
 			OutScalarFieldParam.SetAndAdvance(0.0);
 		}
diff --git a/Engine/Plugins/Experimental/Water/Source/Runtime/Private/NiagaraDataInterfaceWater.cpp b/Engine/Plugins/Experimental/Water/Source/Runtime/Private/NiagaraDataInterfaceWater.cpp
--- a/Engine/Plugins/Experimental/Water/Source/Runtime/Private/NiagaraDataInterfaceWater.cpp
+++ b/Engine/Plugins/Experimental/Water/Source/Runtime/Private/NiagaraDataInterfaceWater.cpp
@@ -152,7 +152,7 @@ bool UNiagaraDataInterfaceWater::PerInstanceTick(void* PerInstanceData, FNiagara
 	return false;
 }
 
-void UNiagaraDataInterfaceWater::GetWaterDataAtPoint(FVectorVMContext& Context)
+void UNiagaraDataInterfaceWater::GetWaterDataAtPoint(FVectorVMExternalFunctionContext& Context)
 {
 	QUICK_SCOPE_CYCLE_COUNTER(NiagaraDataInterfaceWater_GetWaterDataAtPoint);
 
@@ -187,7 +187,7 @@ void UNiagaraDataInterfaceWater::GetWaterDataAtPoint(FVectorVMContext& Context)
 		UE_LOG(LogWater, Warning, TEXT("NiagaraDataInterfaceWater: GetWaterData called with no water body actor set"));
 	}
 
-	for (int32 i = 0; i < Context.NumInstances; ++i)
+	for (int32 i = 0; i < Context.GetNumInstances(); ++i)
 	{
 		FWaterBodyQueryResult QueryResult;
 		
@@ -230,7 +230,7 @@ void UNiagaraDataInterfaceWater::GetWaterDataAtPoint(FVectorVMContext& Context)
 	}
 }
 
-void UNiagaraDataInterfaceWater::GetWaveParamLookupTableOffset(FVectorVMContext& Context)
+void UNiagaraDataInterfaceWater::GetWaveParamLookupTableOffset(FVectorVMExternalFunctionContext& Context)
 {
 	// Inputs
 	VectorVM::FUserPtrHandler<FNDIWater_InstanceData> InstData(Context);
@@ -239,14 +239,14 @@ void UNiagaraDataInterfaceWater::GetWaveParamLookupTableOffset(FVectorVMContext&
 	VectorVM::FExternalFuncRegisterHandler<int> OutLookupTableOffset(Context);
 	if (AWaterBody* Actor = InstData->WaterBodyActor.Get())
 	{
-		for (int32 i = 0; i < Context.NumInstances; ++i)
+		for (int32 i = 0; i < Context.GetNumInstances(); ++i)
 		{
 			*OutLookupTableOffset.GetDestAndAdvance() = Actor->WaterBodyIndex;
 		}
 	}
 	else
 	{
-		for (int32 i = 0; i < Context.NumInstances; ++i)
+		for (int32 i = 0; i < Context.GetNumInstances(); ++i)
 		{
 			*OutLookupTableOffset.GetDestAndAdvance() = 0;
 		}
diff --git a/Engine/Plugins/Experimental/Water/Source/Runtime/Public/NiagaraDataInterfaceWater.h b/Engine/Plugins/Experimental/Water/Source/Runtime/Public/NiagaraDataInterfaceWater.h
--- a/Engine/Plugins/Experimental/Water/Source/Runtime/Public/NiagaraDataInterfaceWater.h
+++ b/Engine/Plugins/Experimental/Water/Source/Runtime/Public/NiagaraDataInterfaceWater.h
@@ -28,9 +28,9 @@ public:
 	virtual bool PerInstanceTick(void* PerInstanceData, FNiagaraSystemInstance* SystemInstance, float DeltaSeconds) override;
 	virtual bool HasPreSimulateTick() const override { return true; }
 
-	void GetWaterDataAtPoint(FVectorVMContext& Context);
+	void GetWaterDataAtPoint(FVectorVMExternalFunctionContext& Context);
 
-	void GetWaveParamLookupTableOffset(FVectorVMContext& Context);
+	void GetWaveParamLookupTableOffset(FVectorVMExternalFunctionContext& Context);
 
 	/** Sets the current water body to be used by this data interface */
 	void SetWaterBody(AWaterBody* InWaterBody) { SourceBody = InWaterBody; }
diff --git a/Engine/Plugins/FX/Niagara/Source/Niagara/Classes/NiagaraDataInterface.h b/Engine/Plugins/FX/Niagara/Source/Niagara/Classes/NiagaraDataInterface.h
--- a/Engine/Plugins/FX/Niagara/Source/Niagara/Classes/NiagaraDataInterface.h
+++ b/Engine/Plugins/FX/Niagara/Source/Niagara/Classes/NiagaraDataInterface.h
@@ -87,7 +87,7 @@ struct NDI_FUNC_BINDER(ClassName, FuncName)\
 	template<typename ... ParamTypes>\
 	static void Bind(UNiagaraDataInterface* Interface, const FVMExternalFunctionBindingInfo& BindingInfo, void* InstanceData, FVMExternalFunction &OutFunc)\
 	{\
-		auto Lambda = [Interface](FVectorVMContext& Context) { static_cast<ClassName*>(Interface)->FuncName<ParamTypes...>(Context); };\
+		auto Lambda = [Interface](FVectorVMExternalFunctionContext& Context) { static_cast<ClassName*>(Interface)->FuncName<ParamTypes...>(Context); };\
 		OutFunc = FVMExternalFunction::CreateLambda(Lambda);\
 	}\
 };
@@ -97,7 +97,7 @@ struct NDI_FUNC_BINDER(ClassName, FuncName)\
 {\
 	static void Bind(UNiagaraDataInterface* Interface, FVMExternalFunction &OutFunc)\
 	{\
-		auto Lambda = [Interface](FVectorVMContext& Context) { static_cast<ClassName*>(Interface)->FuncName(Context); };\
+		auto Lambda = [Interface](FVectorVMExternalFunctionContext& Context) { static_cast<ClassName*>(Interface)->FuncName(Context); };\
 		OutFunc = FVMExternalFunction::CreateLambda(Lambda);\
 	}\
 };
@@ -108,7 +108,7 @@ struct NDI_FUNC_BINDER(ClassName, FuncName)\
 	template <typename ... VarTypes>\
 	static void Bind(UNiagaraDataInterface* Interface, FVMExternalFunction &OutFunc, VarTypes... Var)\
 	{\
-		auto Lambda = [Interface, Var...](FVectorVMContext& Context) { static_cast<ClassName*>(Interface)->FuncName(Context, Var...); };\
+		auto Lambda = [Interface, Var...](FVectorVMExternalFunctionContext& Context) { static_cast<ClassName*>(Interface)->FuncName(Context, Var...); };\
 		OutFunc = FVMExternalFunction::CreateLambda(Lambda);\
 	}\
 };
@@ -508,7 +508,7 @@ private:
 template<typename T>
 struct FNDIParameter
 {
-	FNDIParameter(FVectorVMContext& Context) = delete;
+	FNDIParameter(FVectorVMExternalFunctionContext& Context) = delete;
 	FORCEINLINE void GetAndAdvance(T& OutValue) = delete;
 	FORCEINLINE bool IsConstant() const = delete;
 };
@@ -520,9 +520,9 @@ struct FNDIParameter<FNiagaraRandInfo>
 	VectorVM::FExternalFuncInputHandler<int32> Seed2Param;
 	VectorVM::FExternalFuncInputHandler<int32> Seed3Param;
 	
-	FVectorVMContext& Context;
+	FVectorVMExternalFunctionContext& Context;
 
-	FNDIParameter(FVectorVMContext& InContext)
+	FNDIParameter(FVectorVMExternalFunctionContext& InContext)
 		: Context(InContext)
 	{
 		Seed1Param.Init(Context);
@@ -546,7 +546,7 @@ struct FNDIParameter<FNiagaraRandInfo>
 
 struct FNDIRandomHelper
 {
-	FNDIRandomHelper(FVectorVMContext& InContext)
+	FNDIRandomHelper(FVectorVMExternalFunctionContext& InContext)
 		: Context(InContext)
 		, RandParam(Context)
 	{
@@ -569,7 +569,7 @@ struct FNDIRandomHelper
 	{
 		if (IsDeterministic())
 		{
-			int32 RandomCounter = Context.RandCounters[InstanceIndex]++;
+			int32 RandomCounter = Context.GetRandCounters()[InstanceIndex]++;
 
 			FIntVector4 v = FIntVector4(RandomCounter, RandInfo.Seed1, RandInfo.Seed2, RandInfo.Seed3) * 1664525 + FIntVector4(1013904223);
 
@@ -592,7 +592,7 @@ struct FNDIRandomHelper
 		}
 		else
 		{
-			return FVector4(Context.RandStream.GetFraction(), Context.RandStream.GetFraction(), Context.RandStream.GetFraction(), Context.RandStream.GetFraction());
+			return FVector4(Context.GetRandStream().GetFraction(), Context.GetRandStream().GetFraction(), Context.GetRandStream().GetFraction(), Context.GetRandStream().GetFraction());
 		}
 	}
 
@@ -600,7 +600,7 @@ struct FNDIRandomHelper
 	{
 		if (IsDeterministic())
 		{
-			int32 RandomCounter = Context.RandCounters[InstanceIndex]++;
+			int32 RandomCounter = Context.GetRandCounters()[InstanceIndex]++;
 
 			FIntVector v = FIntVector(RandInfo.Seed1, RandInfo.Seed2, RandomCounter | (RandInfo.Seed3 << 16)) * 1664525 + FIntVector(1013904223);
 
@@ -615,7 +615,7 @@ struct FNDIRandomHelper
 		}
 		else
 		{
-			return FVector(Context.RandStream.GetFraction(), Context.RandStream.GetFraction(), Context.RandStream.GetFraction());
+			return FVector(Context.GetRandStream().GetFraction(), Context.GetRandStream().GetFraction(), Context.GetRandStream().GetFraction());
 		}
 	}
 
@@ -628,7 +628,7 @@ struct FNDIRandomHelper
 		}
 		else
 		{
-			return FVector2D(Context.RandStream.GetFraction(), Context.RandStream.GetFraction());
+			return FVector2D(Context.GetRandStream().GetFraction(), Context.GetRandStream().GetFraction());
 		}
 	}
 
@@ -640,7 +640,7 @@ struct FNDIRandomHelper
 		}
 		else
 		{
-			return Context.RandStream.GetFraction();
+			return Context.GetRandStream().GetFraction();
 		}
 	}
 
@@ -687,7 +687,7 @@ struct FNDIRandomHelper
 		return FVector(1.0f - sqrt0, sqrt0 * (1.0 - r.Y), r.Y * sqrt0);
 	}
 
-	FVectorVMContext& Context;
+	FVectorVMExternalFunctionContext& Context;
 	FNDIParameter<FNiagaraRandInfo> RandParam;
 
 	FNiagaraRandInfo RandInfo;
@@ -698,7 +698,7 @@ template<typename T>
 struct FNDIInputParam
 {
 	VectorVM::FExternalFuncInputHandler<T> Data;
-	FORCEINLINE FNDIInputParam(FVectorVMContext& Context) : Data(Context) {}
+	FORCEINLINE FNDIInputParam(FVectorVMExternalFunctionContext& Context) : Data(Context) {}
 	FORCEINLINE T GetAndAdvance() { return Data.GetAndAdvance(); }
 };
 
@@ -706,7 +706,7 @@ template<>
 struct FNDIInputParam<FNiagaraBool>
 {
 	VectorVM::FExternalFuncInputHandler<FNiagaraBool> Data;
-	FORCEINLINE FNDIInputParam(FVectorVMContext& Context) : Data(Context) {}
+	FORCEINLINE FNDIInputParam(FVectorVMExternalFunctionContext& Context) : Data(Context) {}
 	FORCEINLINE bool GetAndAdvance() { return Data.GetAndAdvance().GetValue(); }
 };
 
@@ -714,7 +714,7 @@ template<>
 struct FNDIInputParam<bool>
 {
 	VectorVM::FExternalFuncInputHandler<FNiagaraBool> Data;
-	FORCEINLINE FNDIInputParam(FVectorVMContext& Context) : Data(Context) {}
+	FORCEINLINE FNDIInputParam(FVectorVMExternalFunctionContext& Context) : Data(Context) {}
 	FORCEINLINE bool GetAndAdvance() { return Data.GetAndAdvance().GetValue(); }
 };
 
@@ -723,7 +723,7 @@ struct FNDIInputParam<FVector2D>
 {
 	VectorVM::FExternalFuncInputHandler<float> X;
 	VectorVM::FExternalFuncInputHandler<float> Y;
-	FORCEINLINE FNDIInputParam(FVectorVMContext& Context) : X(Context), Y(Context) {}
+	FORCEINLINE FNDIInputParam(FVectorVMExternalFunctionContext& Context) : X(Context), Y(Context) {}
 	FORCEINLINE FVector2D GetAndAdvance() { return FVector2D(X.GetAndAdvance(), Y.GetAndAdvance()); }
 };
 
@@ -733,7 +733,7 @@ struct FNDIInputParam<FVector3f>
 	VectorVM::FExternalFuncInputHandler<float> X;
 	VectorVM::FExternalFuncInputHandler<float> Y;
 	VectorVM::FExternalFuncInputHandler<float> Z;
-	FNDIInputParam(FVectorVMContext& Context) : X(Context), Y(Context), Z(Context) {}
+	FNDIInputParam(FVectorVMExternalFunctionContext& Context) : X(Context), Y(Context), Z(Context) {}
 	FORCEINLINE FVector3f GetAndAdvance() { return FVector3f(X.GetAndAdvance(), Y.GetAndAdvance(), Z.GetAndAdvance()); }
 };
 
@@ -746,7 +746,7 @@ struct FNDIInputParam<FVector4>
 	VectorVM::FExternalFuncInputHandler<float> Y;
 	VectorVM::FExternalFuncInputHandler<float> Z;
 	VectorVM::FExternalFuncInputHandler<float> W;
-	FORCEINLINE FNDIInputParam(FVectorVMContext& Context) : X(Context), Y(Context), Z(Context), W(Context) {}
+	FORCEINLINE FNDIInputParam(FVectorVMExternalFunctionContext& Context) : X(Context), Y(Context), Z(Context), W(Context) {}
 	FORCEINLINE FVector4 GetAndAdvance() { return FVector4(X.GetAndAdvance(), Y.GetAndAdvance(), Z.GetAndAdvance(), W.GetAndAdvance()); }
 };
 
@@ -757,7 +757,7 @@ struct FNDIInputParam<FQuat>
 	VectorVM::FExternalFuncInputHandler<float> Y;
 	VectorVM::FExternalFuncInputHandler<float> Z;
 	VectorVM::FExternalFuncInputHandler<float> W;
-	FORCEINLINE FNDIInputParam(FVectorVMContext& Context) : X(Context), Y(Context), Z(Context), W(Context) {}
+	FORCEINLINE FNDIInputParam(FVectorVMExternalFunctionContext& Context) : X(Context), Y(Context), Z(Context), W(Context) {}
 	FORCEINLINE FQuat GetAndAdvance() { return FQuat(X.GetAndAdvance(), Y.GetAndAdvance(), Z.GetAndAdvance(), W.GetAndAdvance()); }
 };
 
@@ -768,7 +768,7 @@ struct FNDIInputParam<FLinearColor>
 	VectorVM::FExternalFuncInputHandler<float> G;
 	VectorVM::FExternalFuncInputHandler<float> B;
 	VectorVM::FExternalFuncInputHandler<float> A;
-	FORCEINLINE FNDIInputParam(FVectorVMContext& Context) : R(Context), G(Context), B(Context), A(Context) {}
+	FORCEINLINE FNDIInputParam(FVectorVMExternalFunctionContext& Context) : R(Context), G(Context), B(Context), A(Context) {}
 	FORCEINLINE FLinearColor GetAndAdvance() { return FLinearColor(R.GetAndAdvance(), G.GetAndAdvance(), B.GetAndAdvance(), A.GetAndAdvance()); }
 };
 
@@ -777,7 +777,7 @@ struct FNDIInputParam<FNiagaraID>
 {
 	VectorVM::FExternalFuncInputHandler<int32> Index;
 	VectorVM::FExternalFuncInputHandler<int32> AcquireTag;
-	FORCEINLINE FNDIInputParam(FVectorVMContext& Context) : Index(Context), AcquireTag(Context) {}
+	FORCEINLINE FNDIInputParam(FVectorVMExternalFunctionContext& Context) : Index(Context), AcquireTag(Context) {}
 	FORCEINLINE FNiagaraID GetAndAdvance() { return FNiagaraID(Index.GetAndAdvance(), AcquireTag.GetAndAdvance()); }
 };
 
@@ -786,7 +786,7 @@ template<typename T>
 struct FNDIOutputParam
 {
 	VectorVM::FExternalFuncRegisterHandler<T> Data;
-	FORCEINLINE FNDIOutputParam(FVectorVMContext& Context) : Data(Context) {}
+	FORCEINLINE FNDIOutputParam(FVectorVMExternalFunctionContext& Context) : Data(Context) {}
 	FORCEINLINE bool IsValid() const { return Data.IsValid();  }
 	FORCEINLINE void SetAndAdvance(T Val) { *Data.GetDestAndAdvance() = Val; }
 };
@@ -795,7 +795,7 @@ template<>
 struct FNDIOutputParam<FNiagaraBool>
 {
 	VectorVM::FExternalFuncRegisterHandler<FNiagaraBool> Data;
-	FORCEINLINE FNDIOutputParam(FVectorVMContext& Context) : Data(Context) {}
+	FORCEINLINE FNDIOutputParam(FVectorVMExternalFunctionContext& Context) : Data(Context) {}
 	FORCEINLINE bool IsValid() const { return Data.IsValid(); }
 	FORCEINLINE void SetAndAdvance(bool Val) { Data.GetDestAndAdvance()->SetValue(Val); }
 };
@@ -804,7 +804,7 @@ template<>
 struct FNDIOutputParam<bool>
 {
 	VectorVM::FExternalFuncRegisterHandler<FNiagaraBool> Data;
-	FORCEINLINE FNDIOutputParam(FVectorVMContext& Context) : Data(Context) {}
+	FORCEINLINE FNDIOutputParam(FVectorVMExternalFunctionContext& Context) : Data(Context) {}
 	FORCEINLINE bool IsValid() const { return Data.IsValid(); }
 	FORCEINLINE void SetAndAdvance(bool Val) { Data.GetDestAndAdvance()->SetValue(Val); }
 };
@@ -814,7 +814,7 @@ struct FNDIOutputParam<FVector2D>
 {
 	VectorVM::FExternalFuncRegisterHandler<float> X;
 	VectorVM::FExternalFuncRegisterHandler<float> Y;
-	FORCEINLINE FNDIOutputParam(FVectorVMContext& Context) : X(Context), Y(Context) {}
+	FORCEINLINE FNDIOutputParam(FVectorVMExternalFunctionContext& Context) : X(Context), Y(Context) {}
 	FORCEINLINE bool IsValid() const { return X.IsValid() || Y.IsValid(); }
 	FORCEINLINE void SetAndAdvance(FVector2D Val)
 	{
@@ -829,7 +829,7 @@ struct FNDIOutputParam<FVector3f>
 	VectorVM::FExternalFuncRegisterHandler<float> X;
 	VectorVM::FExternalFuncRegisterHandler<float> Y;
 	VectorVM::FExternalFuncRegisterHandler<float> Z;
-	FNDIOutputParam(FVectorVMContext& Context) : X(Context), Y(Context), Z(Context) {}
+	FNDIOutputParam(FVectorVMExternalFunctionContext& Context) : X(Context), Y(Context), Z(Context) {}
 	FORCEINLINE bool IsValid() const { return X.IsValid() || Y.IsValid() || Z.IsValid(); }
 	FORCEINLINE void SetAndAdvance(FVector3f Val)
 	{
@@ -847,7 +847,7 @@ struct FNDIOutputParam<FVector4>
 	VectorVM::FExternalFuncRegisterHandler<float> Y;
 	VectorVM::FExternalFuncRegisterHandler<float> Z;
 	VectorVM::FExternalFuncRegisterHandler<float> W;
-	FORCEINLINE FNDIOutputParam(FVectorVMContext& Context) : X(Context), Y(Context), Z(Context), W(Context) {}
+	FORCEINLINE FNDIOutputParam(FVectorVMExternalFunctionContext& Context) : X(Context), Y(Context), Z(Context), W(Context) {}
 	FORCEINLINE bool IsValid() const { return X.IsValid() || Y.IsValid() || Z.IsValid() || W.IsValid(); }
 	FORCEINLINE void SetAndAdvance(FVector4 Val)
 	{
@@ -865,7 +865,7 @@ struct FNDIOutputParam<FQuat>
 	VectorVM::FExternalFuncRegisterHandler<float> Y;
 	VectorVM::FExternalFuncRegisterHandler<float> Z;
 	VectorVM::FExternalFuncRegisterHandler<float> W;
-	FORCEINLINE FNDIOutputParam(FVectorVMContext& Context) : X(Context), Y(Context), Z(Context), W(Context) {}
+	FORCEINLINE FNDIOutputParam(FVectorVMExternalFunctionContext& Context) : X(Context), Y(Context), Z(Context), W(Context) {}
 	FORCEINLINE bool IsValid() const { return X.IsValid() || Y.IsValid() || Z.IsValid() || W.IsValid(); }
 	FORCEINLINE void SetAndAdvance(FQuat Val)
 	{
@@ -896,7 +896,7 @@ struct FNDIOutputParam<FMatrix44f>
 	VectorVM::FExternalFuncRegisterHandler<float> Out14;
 	VectorVM::FExternalFuncRegisterHandler<float> Out15;
 
-	FORCEINLINE FNDIOutputParam(FVectorVMContext& Context) : Out00(Context), Out01(Context), Out02(Context), Out03(Context), Out04(Context), Out05(Context),
+	FORCEINLINE FNDIOutputParam(FVectorVMExternalFunctionContext& Context) : Out00(Context), Out01(Context), Out02(Context), Out03(Context), Out04(Context), Out05(Context),
 		Out06(Context), Out07(Context), Out08(Context), Out09(Context), Out10(Context), Out11(Context), Out12(Context), Out13(Context), Out14(Context), Out15(Context)	{}
 	FORCEINLINE bool IsValid() const { return Out00.IsValid(); }
 	FORCEINLINE void SetAndAdvance(const FMatrix44f& Val)
@@ -927,7 +927,7 @@ struct FNDIOutputParam<FLinearColor>
 	VectorVM::FExternalFuncRegisterHandler<float> G;
 	VectorVM::FExternalFuncRegisterHandler<float> B;
 	VectorVM::FExternalFuncRegisterHandler<float> A;
-	FORCEINLINE FNDIOutputParam(FVectorVMContext& Context) : R(Context), G(Context), B(Context), A(Context) {}
+	FORCEINLINE FNDIOutputParam(FVectorVMExternalFunctionContext& Context) : R(Context), G(Context), B(Context), A(Context) {}
 	FORCEINLINE bool IsValid() const { return R.IsValid() || G.IsValid() || B.IsValid() || A.IsValid(); }
 	FORCEINLINE void SetAndAdvance(FLinearColor Val)
 	{
@@ -943,7 +943,7 @@ struct FNDIOutputParam<FNiagaraID>
 {
 	VectorVM::FExternalFuncRegisterHandler<int32> Index;
 	VectorVM::FExternalFuncRegisterHandler<int32> AcquireTag;
-	FORCEINLINE FNDIOutputParam(FVectorVMContext& Context) : Index(Context), AcquireTag(Context) {}
+	FORCEINLINE FNDIOutputParam(FVectorVMExternalFunctionContext& Context) : Index(Context), AcquireTag(Context) {}
 	FORCEINLINE bool IsValid() const { return Index.IsValid() || AcquireTag.IsValid(); }
 	FORCEINLINE void SetAndAdvance(FNiagaraID Val)
 	{
diff --git a/Engine/Plugins/FX/Niagara/Source/Niagara/Classes/NiagaraDataInterface2DArrayTexture.h b/Engine/Plugins/FX/Niagara/Source/Niagara/Classes/NiagaraDataInterface2DArrayTexture.h
--- a/Engine/Plugins/FX/Niagara/Source/Niagara/Classes/NiagaraDataInterface2DArrayTexture.h
+++ b/Engine/Plugins/FX/Niagara/Source/Niagara/Classes/NiagaraDataInterface2DArrayTexture.h
@@ -38,7 +38,7 @@ public:
 	virtual bool PerInstanceTick(void* PerInstanceData, FNiagaraSystemInstance* SystemInstance, float DeltaSeconds) override;
 	//UNiagaraDataInterface Interface End
 
-	void GetTextureDimensions(FVectorVMContext& Context);
+	void GetTextureDimensions(FVectorVMExternalFunctionContext& Context);
 
 	virtual bool Equals(const UNiagaraDataInterface* Other) const override;
 
diff --git a/Engine/Plugins/FX/Niagara/Source/Niagara/Classes/NiagaraDataInterfaceArrayImpl.h b/Engine/Plugins/FX/Niagara/Source/Niagara/Classes/NiagaraDataInterfaceArrayImpl.h
--- a/Engine/Plugins/FX/Niagara/Source/Niagara/Classes/NiagaraDataInterfaceArrayImpl.h
+++ b/Engine/Plugins/FX/Niagara/Source/Niagara/Classes/NiagaraDataInterfaceArrayImpl.h
@@ -266,48 +266,48 @@ struct FNiagaraDataInterfaceArrayImpl : public INiagaraDataInterfaceArrayImpl
 		if (BindingInfo.Name == FNiagaraDataInterfaceArrayImplHelper::Function_LengthName)
 		{
 			check(BindingInfo.GetNumInputs() == 0 && BindingInfo.GetNumOutputs() == 1);
-			OutFunc = FVMExternalFunction::CreateLambda([this](FVectorVMContext& Context) { this->GetLength(Context); });
+			OutFunc = FVMExternalFunction::CreateLambda([this](FVectorVMExternalFunctionContext& Context) { this->GetLength(Context); });
 		}
 		else if (BindingInfo.Name == FNiagaraDataInterfaceArrayImplHelper::Function_IsValidIndexName)
 		{
 			check(BindingInfo.GetNumInputs() == 1 && BindingInfo.GetNumOutputs() == 1);
-			OutFunc = FVMExternalFunction::CreateLambda([this](FVectorVMContext& Context) { this->IsValidIndex(Context); });
+			OutFunc = FVMExternalFunction::CreateLambda([this](FVectorVMExternalFunctionContext& Context) { this->IsValidIndex(Context); });
 		}
 		else if (BindingInfo.Name == FNiagaraDataInterfaceArrayImplHelper::Function_LastIndexName)
 		{
 			check(BindingInfo.GetNumInputs() == 0 && BindingInfo.GetNumOutputs() == 1);
-			OutFunc = FVMExternalFunction::CreateLambda([this](FVectorVMContext& Context) { this->GetLastIndex(Context); });
+			OutFunc = FVMExternalFunction::CreateLambda([this](FVectorVMExternalFunctionContext& Context) { this->GetLastIndex(Context); });
 		}
 		else if (BindingInfo.Name == FNiagaraDataInterfaceArrayImplHelper::Function_GetName)
 		{
 			// Note: Outputs is variable based upon type
 			//check(BindingInfo.GetNumInputs() == 1 && BindingInfo.GetNumOutputs() == 1);
-			OutFunc = FVMExternalFunction::CreateLambda([this](FVectorVMContext& Context) { this->GetValue(Context); });
+			OutFunc = FVMExternalFunction::CreateLambda([this](FVectorVMExternalFunctionContext& Context) { this->GetValue(Context); });
 		}
 		// Mutable functions
 		else if (BindingInfo.Name == FNiagaraDataInterfaceArrayImplHelper::Function_ClearName)
 		{
 			check(BindingInfo.GetNumInputs() == 0 && BindingInfo.GetNumOutputs() == 0);
-			OutFunc = FVMExternalFunction::CreateLambda([this](FVectorVMContext& Context) { this->Clear(Context); });
+			OutFunc = FVMExternalFunction::CreateLambda([this](FVectorVMExternalFunctionContext& Context) { this->Clear(Context); });
 		}
 		else if (BindingInfo.Name == FNiagaraDataInterfaceArrayImplHelper::Function_ResizeName)
 		{
 			check(BindingInfo.GetNumInputs() == 1 && BindingInfo.GetNumOutputs() == 0);
-			OutFunc = FVMExternalFunction::CreateLambda([this](FVectorVMContext& Context) { this->Resize(Context); });
+			OutFunc = FVMExternalFunction::CreateLambda([this](FVectorVMExternalFunctionContext& Context) { this->Resize(Context); });
 		}
 		else if (BindingInfo.Name == FNiagaraDataInterfaceArrayImplHelper::Function_SetArrayElemName)
 		{
-			OutFunc = FVMExternalFunction::CreateLambda([this](FVectorVMContext& Context) { this->SetValue(Context); });
+			OutFunc = FVMExternalFunction::CreateLambda([this](FVectorVMExternalFunctionContext& Context) { this->SetValue(Context); });
 		}
 		else if (BindingInfo.Name == FNiagaraDataInterfaceArrayImplHelper::Function_AddName)
 		{
 			// Note: Inputs is variable based upon type
-			OutFunc = FVMExternalFunction::CreateLambda([this](FVectorVMContext& Context) { this->PushValue(Context); });
+			OutFunc = FVMExternalFunction::CreateLambda([this](FVectorVMExternalFunctionContext& Context) { this->PushValue(Context); });
 		}
 		else if (BindingInfo.Name == FNiagaraDataInterfaceArrayImplHelper::Function_RemoveLastElemName)
 		{
 			// Note: Outputs is variable based upon type
-			OutFunc = FVMExternalFunction::CreateLambda([this](FVectorVMContext& Context) { this->PopValue(Context); });
+			OutFunc = FVMExternalFunction::CreateLambda([this](FVectorVMExternalFunctionContext& Context) { this->PopValue(Context); });
 		}
 	}
 
@@ -502,20 +502,20 @@ struct FNiagaraDataInterfaceArrayImpl : public INiagaraDataInterfaceArrayImpl
 		//static_cast<const FNiagaraDataInterfaceParametersCS_ArrayImpl*>(Base)->Unset(RHICmdList, Context);
 	}
 
-	void GetLength(FVectorVMContext& Context)
+	void GetLength(FVectorVMExternalFunctionContext& Context)
 	{
 		FNDIOutputParam<int32> OutValue(Context);
 
 		Owner->ArrayRWGuard.ReadLock();
 		const int32 Num = Data.Num();
 		Owner->ArrayRWGuard.ReadUnlock();
-		for (int32 i = 0; i < Context.NumInstances; ++i)
+		for (int32 i = 0; i < Context.GetNumInstances(); ++i)
 		{
 			OutValue.SetAndAdvance(Num);
 		}
 	}
 
-	void IsValidIndex(FVectorVMContext& Context)
+	void IsValidIndex(FVectorVMExternalFunctionContext& Context)
 	{
 		FNDIInputParam<int32> IndexParam(Context);
 		FNDIOutputParam<FNiagaraBool> OutValue(Context);
@@ -523,27 +523,27 @@ struct FNiagaraDataInterfaceArrayImpl : public INiagaraDataInterfaceArrayImpl
 		Owner->ArrayRWGuard.ReadLock();
 		const int32 Num = Data.Num();
 		Owner->ArrayRWGuard.ReadUnlock();
-		for (int32 i = 0; i < Context.NumInstances; ++i)
+		for (int32 i = 0; i < Context.GetNumInstances(); ++i)
 		{
 			const int32 Index = IndexParam.GetAndAdvance();
 			OutValue.SetAndAdvance((Index >= 0) && (Index < Num));
 		}
 	}
 
-	void GetLastIndex(FVectorVMContext& Context)
+	void GetLastIndex(FVectorVMExternalFunctionContext& Context)
 	{
 		FNDIOutputParam<int32> OutValue(Context);
 
 		Owner->ArrayRWGuard.ReadLock();
 		const int32 Num = Data.Num() - 1;
 		Owner->ArrayRWGuard.ReadUnlock();
-		for (int32 i = 0; i < Context.NumInstances; ++i)
+		for (int32 i = 0; i < Context.GetNumInstances(); ++i)
 		{
 			OutValue.SetAndAdvance(Num);
 		}
 	}
 
-	void GetValue(FVectorVMContext& Context)
+	void GetValue(FVectorVMExternalFunctionContext& Context)
 	{
 		FNDIInputParam<int32> IndexParam(Context);
 		FNDIOutputParam<typename FNDIArrayImplHelper<TArrayType>::TVMArrayType> OutValue(Context);
@@ -552,7 +552,7 @@ struct FNiagaraDataInterfaceArrayImpl : public INiagaraDataInterfaceArrayImpl
 		const int32 Num = Data.Num() - 1;
 		if (Num >= 0)
 		{
-			for (int32 i = 0; i < Context.NumInstances; ++i)
+			for (int32 i = 0; i < Context.GetNumInstances(); ++i)
 			{
 				const int32 Index = FMath::Clamp(IndexParam.GetAndAdvance(), 0, Num);
 				OutValue.SetAndAdvance(Data[Index]);
@@ -561,28 +561,28 @@ struct FNiagaraDataInterfaceArrayImpl : public INiagaraDataInterfaceArrayImpl
 		else
 		{
 			const TArrayType DefaultValue = FNDIArrayImplHelper<TArrayType>::GetDefaultValue();
-			for (int32 i = 0; i < Context.NumInstances; ++i)
+			for (int32 i = 0; i < Context.GetNumInstances(); ++i)
 			{
 				OutValue.SetAndAdvance(DefaultValue);
 			}
 		}
 	}
 
-	void Clear(FVectorVMContext& Context)
+	void Clear(FVectorVMExternalFunctionContext& Context)
 	{
 		//-TODO: This dirties the GPU data
-		ensureMsgf(Context.NumInstances == 1, TEXT("Setting the number of values in an array with more than one instance, which doesn't make sense"));
+		ensureMsgf(Context.GetNumInstances() == 1, TEXT("Setting the number of values in an array with more than one instance, which doesn't make sense"));
 
 		FRWScopeLock WriteLock(Owner->ArrayRWGuard, SLT_Write);
 		Data.Reset();
 	}
 
-	void Resize(FVectorVMContext& Context)
+	void Resize(FVectorVMExternalFunctionContext& Context)
 	{
 		//-TODO: This dirties the GPU data
 		FNDIInputParam<int32> NewNumParam(Context);
 
-		ensureMsgf(Context.NumInstances == 1, TEXT("Setting the number of values in an array with more than one instance, which doesn't make sense"));
+		ensureMsgf(Context.GetNumInstances() == 1, TEXT("Setting the number of values in an array with more than one instance, which doesn't make sense"));
 
 		FRWScopeLock WriteLock(Owner->ArrayRWGuard, SLT_Write);
 		const int32 OldNum = Data.Num();
@@ -599,14 +599,14 @@ struct FNiagaraDataInterfaceArrayImpl : public INiagaraDataInterfaceArrayImpl
 		}
 	}
 
-	void SetValue(FVectorVMContext& Context)
+	void SetValue(FVectorVMExternalFunctionContext& Context)
 	{
 		//-TODO: This dirties the GPU data
 		FNDIInputParam<int32> IndexParam(Context);
 		FNDIInputParam<typename FNDIArrayImplHelper<TArrayType>::TVMArrayType> InValue(Context);
 
 		FRWScopeLock WriteLock(Owner->ArrayRWGuard, SLT_Write);
-		for (int32 i = 0; i < Context.NumInstances; ++i)
+		for (int32 i = 0; i < Context.GetNumInstances(); ++i)
 		{
 			const int32 Index = IndexParam.GetAndAdvance();
 			const TArrayType Value = InValue.GetAndAdvance();
@@ -618,7 +618,7 @@ struct FNiagaraDataInterfaceArrayImpl : public INiagaraDataInterfaceArrayImpl
 		}
 	}
 
-	void PushValue(FVectorVMContext& Context)
+	void PushValue(FVectorVMExternalFunctionContext& Context)
 	{
 		//-TODO: This dirties the GPU data
 		FNDIInputParam<FNiagaraBool> InSkipExecute(Context);
@@ -627,7 +627,7 @@ struct FNiagaraDataInterfaceArrayImpl : public INiagaraDataInterfaceArrayImpl
 		const int32 MaxElements = Owner->MaxElements > 0 ? Owner->MaxElements : kSafeMaxElements;
 
 		FRWScopeLock WriteLock(Owner->ArrayRWGuard, SLT_Write);
-		for (int32 i = 0; i < Context.NumInstances; ++i)
+		for (int32 i = 0; i < Context.GetNumInstances(); ++i)
 		{
 			const bool bSkipExecute = InSkipExecute.GetAndAdvance();
 			const TArrayType Value = InValue.GetAndAdvance();
@@ -638,7 +638,7 @@ struct FNiagaraDataInterfaceArrayImpl : public INiagaraDataInterfaceArrayImpl
 		}
 	}
 
-	void PopValue(FVectorVMContext& Context)
+	void PopValue(FVectorVMExternalFunctionContext& Context)
 	{
 		//-TODO: This dirties the GPU data
 		FNDIInputParam<FNiagaraBool> InSkipExecute(Context);
@@ -647,7 +647,7 @@ struct FNiagaraDataInterfaceArrayImpl : public INiagaraDataInterfaceArrayImpl
 		const TArrayType DefaultValue = FNDIArrayImplHelper<TArrayType>::GetDefaultValue();
 
 		FRWScopeLock WriteLock(Owner->ArrayRWGuard, SLT_Write);
-		for (int32 i=0; i < Context.NumInstances; ++i)
+		for (int32 i=0; i < Context.GetNumInstances(); ++i)
 		{
 			const bool bSkipExecute = InSkipExecute.GetAndAdvance();
 			if (bSkipExecute || (Data.Num() == 0))
diff --git a/Engine/Plugins/FX/Niagara/Source/Niagara/Classes/NiagaraDataInterfaceAudioOscilloscope.h b/Engine/Plugins/FX/Niagara/Source/Niagara/Classes/NiagaraDataInterfaceAudioOscilloscope.h
--- a/Engine/Plugins/FX/Niagara/Source/Niagara/Classes/NiagaraDataInterfaceAudioOscilloscope.h
+++ b/Engine/Plugins/FX/Niagara/Source/Niagara/Classes/NiagaraDataInterfaceAudioOscilloscope.h
@@ -111,8 +111,8 @@ public:
 	float ScopeInMilliseconds;
 
 	//VM function overrides:
-	void SampleAudio(FVectorVMContext& Context);
-	void GetNumChannels(FVectorVMContext& Context);
+	void SampleAudio(FVectorVMExternalFunctionContext& Context);
+	void GetNumChannels(FVectorVMExternalFunctionContext& Context);
 
 	virtual void GetFunctions(TArray<FNiagaraFunctionSignature>& OutFunctions)override;
 	virtual void GetVMExternalFunction(const FVMExternalFunctionBindingInfo& BindingInfo, void* InstanceData, FVMExternalFunction &OutFunc) override;
diff --git a/Engine/Plugins/FX/Niagara/Source/Niagara/Classes/NiagaraDataInterfaceAudioPlayer.h b/Engine/Plugins/FX/Niagara/Source/Niagara/Classes/NiagaraDataInterfaceAudioPlayer.h
--- a/Engine/Plugins/FX/Niagara/Source/Niagara/Classes/NiagaraDataInterfaceAudioPlayer.h
+++ b/Engine/Plugins/FX/Niagara/Source/Niagara/Classes/NiagaraDataInterfaceAudioPlayer.h
@@ -103,16 +103,16 @@ public:
 	virtual bool HasPostSimulateTick() const override { return true; }
 	//UNiagaraDataInterface Interface
 
-	virtual void PlayOneShotAudio(FVectorVMContext& Context);
-	virtual void PlayPersistentAudio(FVectorVMContext& Context);
-	virtual void SetParameterBool(FVectorVMContext& Context);
-	virtual void SetParameterInteger(FVectorVMContext& Context);
-	virtual void SetParameterFloat(FVectorVMContext& Context);
-	virtual void UpdateVolume(FVectorVMContext& Context);
-	virtual void UpdatePitch(FVectorVMContext& Context);
-	virtual void UpdateLocation(FVectorVMContext& Context);
-	virtual void UpdateRotation(FVectorVMContext& Context);
-	virtual void SetPausedState(FVectorVMContext& Context);
+	virtual void PlayOneShotAudio(FVectorVMExternalFunctionContext& Context);
+	virtual void PlayPersistentAudio(FVectorVMExternalFunctionContext& Context);
+	virtual void SetParameterBool(FVectorVMExternalFunctionContext& Context);
+	virtual void SetParameterInteger(FVectorVMExternalFunctionContext& Context);
+	virtual void SetParameterFloat(FVectorVMExternalFunctionContext& Context);
+	virtual void UpdateVolume(FVectorVMExternalFunctionContext& Context);
+	virtual void UpdatePitch(FVectorVMExternalFunctionContext& Context);
+	virtual void UpdateLocation(FVectorVMExternalFunctionContext& Context);
+	virtual void UpdateRotation(FVectorVMExternalFunctionContext& Context);
+	virtual void SetPausedState(FVectorVMExternalFunctionContext& Context);
 
 protected:
 	virtual bool CopyToInternal(UNiagaraDataInterface* Destination) const override;
diff --git a/Engine/Plugins/FX/Niagara/Source/Niagara/Classes/NiagaraDataInterfaceAudioSpectrum.h b/Engine/Plugins/FX/Niagara/Source/Niagara/Classes/NiagaraDataInterfaceAudioSpectrum.h
--- a/Engine/Plugins/FX/Niagara/Source/Niagara/Classes/NiagaraDataInterfaceAudioSpectrum.h
+++ b/Engine/Plugins/FX/Niagara/Source/Niagara/Classes/NiagaraDataInterfaceAudioSpectrum.h
@@ -159,8 +159,8 @@ public:
 	float NoiseFloorDb;
 
 	//VM function overrides:
-	void GetSpectrumValue(FVectorVMContext& Context);
-	void GetNumChannels(FVectorVMContext& Context);
+	void GetSpectrumValue(FVectorVMExternalFunctionContext& Context);
+	void GetNumChannels(FVectorVMExternalFunctionContext& Context);
 
 	virtual void GetFunctions(TArray<FNiagaraFunctionSignature>& OutFunctions) override;
 	virtual void GetVMExternalFunction(const FVMExternalFunctionBindingInfo& BindingInfo, void* InstanceData, FVMExternalFunction &OutFunc) override;
diff --git a/Engine/Plugins/FX/Niagara/Source/Niagara/Classes/NiagaraDataInterfaceCamera.h b/Engine/Plugins/FX/Niagara/Source/Niagara/Classes/NiagaraDataInterfaceCamera.h
--- a/Engine/Plugins/FX/Niagara/Source/Niagara/Classes/NiagaraDataInterfaceCamera.h
+++ b/Engine/Plugins/FX/Niagara/Source/Niagara/Classes/NiagaraDataInterfaceCamera.h
@@ -64,14 +64,14 @@ public:
 #endif
 	//UNiagaraDataInterface Interface
 
-	void CalculateParticleDistances(FVectorVMContext& Context);
-	void GetClosestParticles(FVectorVMContext& Context);
-	void GetCameraFOV(FVectorVMContext& Context);
-	void GetCameraProperties(FVectorVMContext& Context);
-	void GetViewPropertiesGPU(FVectorVMContext& Context);
-	void GetClipSpaceTransformsGPU(FVectorVMContext& Context);
-	void GetViewSpaceTransformsGPU(FVectorVMContext& Context);
-	void GetTAAJitter(FVectorVMContext& Context);
+	void CalculateParticleDistances(FVectorVMExternalFunctionContext& Context);
+	void GetClosestParticles(FVectorVMExternalFunctionContext& Context);
+	void GetCameraFOV(FVectorVMExternalFunctionContext& Context);
+	void GetCameraProperties(FVectorVMExternalFunctionContext& Context);
+	void GetViewPropertiesGPU(FVectorVMExternalFunctionContext& Context);
+	void GetClipSpaceTransformsGPU(FVectorVMExternalFunctionContext& Context);
+	void GetViewSpaceTransformsGPU(FVectorVMExternalFunctionContext& Context);
+	void GetTAAJitter(FVectorVMExternalFunctionContext& Context);
 
 protected:
 	virtual bool CopyToInternal(UNiagaraDataInterface* Destination) const override;
diff --git a/Engine/Plugins/FX/Niagara/Source/Niagara/Classes/NiagaraDataInterfaceCollisionQuery.h b/Engine/Plugins/FX/Niagara/Source/Niagara/Classes/NiagaraDataInterfaceCollisionQuery.h
--- a/Engine/Plugins/FX/Niagara/Source/Niagara/Classes/NiagaraDataInterfaceCollisionQuery.h
+++ b/Engine/Plugins/FX/Niagara/Source/Niagara/Classes/NiagaraDataInterfaceCollisionQuery.h
@@ -58,10 +58,10 @@ public:
 	virtual void PushToRenderThreadImpl() override;
 
 	// VM functions
-	void PerformQuerySyncCPU(FVectorVMContext& Context);
-	void PerformQueryAsyncCPU(FVectorVMContext& Context);
-	void QuerySceneDepth(FVectorVMContext& Context);
-	void QueryMeshDistanceField(FVectorVMContext& Context);
+	void PerformQuerySyncCPU(FVectorVMExternalFunctionContext& Context);
+	void PerformQueryAsyncCPU(FVectorVMExternalFunctionContext& Context);
+	void QuerySceneDepth(FVectorVMExternalFunctionContext& Context);
+	void QueryMeshDistanceField(FVectorVMExternalFunctionContext& Context);
 
 	virtual bool CanExecuteOnTarget(ENiagaraSimTarget Target) const override { return true; }
 	virtual bool RequiresDistanceFieldData() const override { return true; }
diff --git a/Engine/Plugins/FX/Niagara/Source/Niagara/Classes/NiagaraDataInterfaceColorCurve.h b/Engine/Plugins/FX/Niagara/Source/Niagara/Classes/NiagaraDataInterfaceColorCurve.h
--- a/Engine/Plugins/FX/Niagara/Source/Niagara/Classes/NiagaraDataInterfaceColorCurve.h
+++ b/Engine/Plugins/FX/Niagara/Source/Niagara/Classes/NiagaraDataInterfaceColorCurve.h
@@ -53,7 +53,7 @@ public:
 	virtual void GetVMExternalFunction(const FVMExternalFunctionBindingInfo& BindingInfo, void* InstanceData, FVMExternalFunction &OutFunc) override;
 
 	template<typename UseLUT>
-	void SampleCurve(FVectorVMContext& Context);
+	void SampleCurve(FVectorVMExternalFunctionContext& Context);
 
 	virtual bool Equals(const UNiagaraDataInterface* Other) const override;
 
diff --git a/Engine/Plugins/FX/Niagara/Source/Niagara/Classes/NiagaraDataInterfaceCubeTexture.h b/Engine/Plugins/FX/Niagara/Source/Niagara/Classes/NiagaraDataInterfaceCubeTexture.h
--- a/Engine/Plugins/FX/Niagara/Source/Niagara/Classes/NiagaraDataInterfaceCubeTexture.h
+++ b/Engine/Plugins/FX/Niagara/Source/Niagara/Classes/NiagaraDataInterfaceCubeTexture.h
@@ -38,8 +38,8 @@ public:
 	virtual bool PerInstanceTick(void* PerInstanceData, FNiagaraSystemInstance* SystemInstance, float DeltaSeconds) override;
 	//UNiagaraDataInterface Interface End
 
-	void SampleCubeTexture(FVectorVMContext& Context);
-	void GetTextureDimensions(FVectorVMContext& Context);
+	void SampleCubeTexture(FVectorVMExternalFunctionContext& Context);
+	void GetTextureDimensions(FVectorVMExternalFunctionContext& Context);
 
 	virtual bool Equals(const UNiagaraDataInterface* Other) const override;
 
diff --git a/Engine/Plugins/FX/Niagara/Source/Niagara/Classes/NiagaraDataInterfaceCurlNoise.h b/Engine/Plugins/FX/Niagara/Source/Niagara/Classes/NiagaraDataInterfaceCurlNoise.h
--- a/Engine/Plugins/FX/Niagara/Source/Niagara/Classes/NiagaraDataInterfaceCurlNoise.h
+++ b/Engine/Plugins/FX/Niagara/Source/Niagara/Classes/NiagaraDataInterfaceCurlNoise.h
@@ -38,7 +38,7 @@ public:
 	virtual bool Equals(const UNiagaraDataInterface* Other) const override;
 	//UNiagaraDataInterface Interface End
 
-	void SampleNoiseField(FVectorVMContext& Context);
+	void SampleNoiseField(FVectorVMExternalFunctionContext& Context);
 
 	// GPU sim functionality
 #if WITH_EDITORONLY_DATA
diff --git a/Engine/Plugins/FX/Niagara/Source/Niagara/Classes/NiagaraDataInterfaceCurve.h b/Engine/Plugins/FX/Niagara/Source/Niagara/Classes/NiagaraDataInterfaceCurve.h
--- a/Engine/Plugins/FX/Niagara/Source/Niagara/Classes/NiagaraDataInterfaceCurve.h
+++ b/Engine/Plugins/FX/Niagara/Source/Niagara/Classes/NiagaraDataInterfaceCurve.h
@@ -37,7 +37,7 @@ public:
 	virtual void GetVMExternalFunction(const FVMExternalFunctionBindingInfo& BindingInfo, void* InstanceData, FVMExternalFunction &OutFunc) override;
 
 	template<typename UseLUT>
-	void SampleCurve(FVectorVMContext& Context);
+	void SampleCurve(FVectorVMExternalFunctionContext& Context);
 	
 	virtual bool Equals(const UNiagaraDataInterface* Other) const override;
 
diff --git a/Engine/Plugins/FX/Niagara/Source/Niagara/Classes/NiagaraDataInterfaceExport.h b/Engine/Plugins/FX/Niagara/Source/Niagara/Classes/NiagaraDataInterfaceExport.h
--- a/Engine/Plugins/FX/Niagara/Source/Niagara/Classes/NiagaraDataInterfaceExport.h
+++ b/Engine/Plugins/FX/Niagara/Source/Niagara/Classes/NiagaraDataInterfaceExport.h
@@ -98,8 +98,8 @@ public:
 
 	virtual bool HasInternalAttributeReads(const UNiagaraEmitter* OwnerEmitter, const UNiagaraEmitter* Provider) const override { return OwnerEmitter == Provider; };
 
-	virtual void StoreData(FVectorVMContext& Context);
-	virtual void ExportData(FVectorVMContext& Context);
+	virtual void StoreData(FVectorVMExternalFunctionContext& Context);
+	virtual void ExportData(FVectorVMExternalFunctionContext& Context);
 
 protected:
 	virtual bool CopyToInternal(UNiagaraDataInterface* Destination) const override;
diff --git a/Engine/Plugins/FX/Niagara/Source/Niagara/Classes/NiagaraDataInterfaceGrid2DCollection.h b/Engine/Plugins/FX/Niagara/Source/Niagara/Classes/NiagaraDataInterfaceGrid2DCollection.h
--- a/Engine/Plugins/FX/Niagara/Source/Niagara/Classes/NiagaraDataInterfaceGrid2DCollection.h
+++ b/Engine/Plugins/FX/Niagara/Source/Niagara/Classes/NiagaraDataInterfaceGrid2DCollection.h
@@ -202,11 +202,11 @@ public:
 	UFUNCTION(BlueprintCallable, Category = Niagara)
 	virtual void GetTextureSize(const UNiagaraComponent *Component, int &SizeX, int &SizeY);
 
-	void GetWorldBBoxSize(FVectorVMContext& Context);
-	void GetCellSize(FVectorVMContext& Context);
-	void GetNumCells(FVectorVMContext& Context);
-	void SetNumCells(FVectorVMContext& Context);
-	void GetAttributeIndex(FVectorVMContext& Context, const FName& InName, int32 NumChannels);
+	void GetWorldBBoxSize(FVectorVMExternalFunctionContext& Context);
+	void GetCellSize(FVectorVMExternalFunctionContext& Context);
+	void GetNumCells(FVectorVMExternalFunctionContext& Context);
+	void SetNumCells(FVectorVMExternalFunctionContext& Context);
+	void GetAttributeIndex(FVectorVMExternalFunctionContext& Context, const FName& InName, int32 NumChannels);
 
 	static const FString GridName;
 	static const FString OutputGridName;
diff --git a/Engine/Plugins/FX/Niagara/Source/Niagara/Classes/NiagaraDataInterfaceGrid3DCollection.h b/Engine/Plugins/FX/Niagara/Source/Niagara/Classes/NiagaraDataInterfaceGrid3DCollection.h
--- a/Engine/Plugins/FX/Niagara/Source/Niagara/Classes/NiagaraDataInterfaceGrid3DCollection.h
+++ b/Engine/Plugins/FX/Niagara/Source/Niagara/Classes/NiagaraDataInterfaceGrid3DCollection.h
@@ -198,13 +198,13 @@ public:
 	UFUNCTION(BlueprintCallable, Category = Niagara)
 	virtual void GetTextureSize(const UNiagaraComponent *Component, int &SizeX, int &SizeY, int &SizeZ);	
 
-	void GetWorldBBoxSize(FVectorVMContext& Context);
-	void GetCellSize(FVectorVMContext& Context);
+	void GetWorldBBoxSize(FVectorVMExternalFunctionContext& Context);
+	void GetCellSize(FVectorVMExternalFunctionContext& Context);
 
-	void GetNumCells(FVectorVMContext& Context);
-	void SetNumCells(FVectorVMContext& Context);
+	void GetNumCells(FVectorVMExternalFunctionContext& Context);
+	void SetNumCells(FVectorVMExternalFunctionContext& Context);
 
-	void GetAttributeIndex(FVectorVMContext& Context, const FName& InName, int32 NumChannels);
+	void GetAttributeIndex(FVectorVMExternalFunctionContext& Context, const FName& InName, int32 NumChannels);
 
 	static const FString NumTilesName;
 	static const FString OneOverNumTilesName;
diff --git a/Engine/Plugins/FX/Niagara/Source/Niagara/Classes/NiagaraDataInterfaceIntRenderTarget2D.h b/Engine/Plugins/FX/Niagara/Source/Niagara/Classes/NiagaraDataInterfaceIntRenderTarget2D.h
--- a/Engine/Plugins/FX/Niagara/Source/Niagara/Classes/NiagaraDataInterfaceIntRenderTarget2D.h
+++ b/Engine/Plugins/FX/Niagara/Source/Niagara/Classes/NiagaraDataInterfaceIntRenderTarget2D.h
@@ -45,8 +45,8 @@ public:
 
 	bool UpdateInstanceTexture(FNiagaraSystemInstance* SystemInstance, struct FNDIIntRenderTarget2DInstanceData_GameThread* InstanceData);
 
-	void VMGetSize(FVectorVMContext& Context);
-	void VMSetSize(FVectorVMContext& Context);
+	void VMGetSize(FVectorVMExternalFunctionContext& Context);
+	void VMSetSize(FVectorVMExternalFunctionContext& Context);
 
 	UPROPERTY(EditAnywhere, Category = "Render Target")
 	FIntPoint Size = FIntPoint::ZeroValue;
diff --git a/Engine/Plugins/FX/Niagara/Source/Niagara/Classes/NiagaraDataInterfaceMeshRendererInfo.h b/Engine/Plugins/FX/Niagara/Source/Niagara/Classes/NiagaraDataInterfaceMeshRendererInfo.h
--- a/Engine/Plugins/FX/Niagara/Source/Niagara/Classes/NiagaraDataInterfaceMeshRendererInfo.h
+++ b/Engine/Plugins/FX/Niagara/Source/Niagara/Classes/NiagaraDataInterfaceMeshRendererInfo.h
@@ -54,8 +54,8 @@ protected:
 	virtual bool CopyToInternal(UNiagaraDataInterface* Destination) const override;
 	virtual void PushToRenderThreadImpl() override;
 
-	void GetNumMeshes(FVectorVMContext& Context);
-	void GetMeshLocalBounds(FVectorVMContext& Context);
+	void GetNumMeshes(FVectorVMExternalFunctionContext& Context);
+	void GetMeshLocalBounds(FVectorVMExternalFunctionContext& Context);
 
 	/** The name of the mesh renderer */
 	UPROPERTY(EditAnywhere, Category = "Source")
diff --git a/Engine/Plugins/FX/Niagara/Source/Niagara/Classes/NiagaraDataInterfaceNeighborGrid3D.h b/Engine/Plugins/FX/Niagara/Source/Niagara/Classes/NiagaraDataInterfaceNeighborGrid3D.h
--- a/Engine/Plugins/FX/Niagara/Source/Niagara/Classes/NiagaraDataInterfaceNeighborGrid3D.h
+++ b/Engine/Plugins/FX/Niagara/Source/Niagara/Classes/NiagaraDataInterfaceNeighborGrid3D.h
@@ -106,10 +106,10 @@ public:
 	virtual bool HasPreSimulateTick() const override { return true; }
 	//~ UNiagaraDataInterface interface END
 
-	void GetWorldBBoxSize(FVectorVMContext& Context);
-	void GetNumCells(FVectorVMContext& Context);
-	void GetMaxNeighborsPerCell(FVectorVMContext& Context);
-	void SetNumCells(FVectorVMContext& Context);
+	void GetWorldBBoxSize(FVectorVMExternalFunctionContext& Context);
+	void GetNumCells(FVectorVMExternalFunctionContext& Context);
+	void GetMaxNeighborsPerCell(FVectorVMExternalFunctionContext& Context);
+	void SetNumCells(FVectorVMExternalFunctionContext& Context);
 
 	static const FName SetNumCellsFunctionName;
 
diff --git a/Engine/Plugins/FX/Niagara/Source/Niagara/Classes/NiagaraDataInterfaceOcclusion.h b/Engine/Plugins/FX/Niagara/Source/Niagara/Classes/NiagaraDataInterfaceOcclusion.h
--- a/Engine/Plugins/FX/Niagara/Source/Niagara/Classes/NiagaraDataInterfaceOcclusion.h
+++ b/Engine/Plugins/FX/Niagara/Source/Niagara/Classes/NiagaraDataInterfaceOcclusion.h
@@ -28,8 +28,8 @@ public:
 	virtual bool RequiresDepthBuffer() const override { return true; }
 	//UNiagaraDataInterface Interface
 
-	void QueryOcclusionFactorGPU(FVectorVMContext& Context);
-	void QueryOcclusionFactorCircleGPU(FVectorVMContext& Context);
+	void QueryOcclusionFactorGPU(FVectorVMExternalFunctionContext& Context);
+	void QueryOcclusionFactorCircleGPU(FVectorVMExternalFunctionContext& Context);
 	
 private:
 	static const FName GetCameraOcclusionRectangleName;
diff --git a/Engine/Plugins/FX/Niagara/Source/Niagara/Classes/NiagaraDataInterfaceParticleRead.h b/Engine/Plugins/FX/Niagara/Source/Niagara/Classes/NiagaraDataInterfaceParticleRead.h
--- a/Engine/Plugins/FX/Niagara/Source/Niagara/Classes/NiagaraDataInterfaceParticleRead.h
+++ b/Engine/Plugins/FX/Niagara/Source/Niagara/Classes/NiagaraDataInterfaceParticleRead.h
@@ -49,29 +49,29 @@ public:
 	virtual bool HasInternalAttributeReads(const UNiagaraEmitter* OwnerEmitter, const UNiagaraEmitter* Provider) const override;
 	//UNiagaraDataInterface Interface End
 
-	void GetNumSpawnedParticles(FVectorVMContext& Context);
-	void GetSpawnedIDAtIndex(FVectorVMContext& Context);
-	void GetNumParticles(FVectorVMContext& Context);
-	void GetParticleIndex(FVectorVMContext& Context);
-	void GetParticleIndexFromIDTable(FVectorVMContext& Context);
-	void ReadInt(FVectorVMContext& Context, FName AttributeToRead);
-	void ReadBool(FVectorVMContext& Context, FName AttributeToRead);
-	void ReadFloat(FVectorVMContext& Context, FName AttributeToRead);
-	void ReadVector2(FVectorVMContext& Context, FName AttributeToRead);
-	void ReadVector3(FVectorVMContext& Context, FName AttributeToRead);
-	void ReadVector4(FVectorVMContext& Context, FName AttributeToRead);
-	void ReadColor(FVectorVMContext& Context, FName AttributeToRead);
-	void ReadQuat(FVectorVMContext& Context, FName AttributeToRead);
-	void ReadID(FVectorVMContext& Context, FName AttributeToRead);
-	void ReadIntByIndex(FVectorVMContext& Context, FName AttributeToRead);
-	void ReadBoolByIndex(FVectorVMContext& Context, FName AttributeToRead);
-	void ReadFloatByIndex(FVectorVMContext& Context, FName AttributeToRead);
-	void ReadVector2ByIndex(FVectorVMContext& Context, FName AttributeToRead);
-	void ReadVector3ByIndex(FVectorVMContext& Context, FName AttributeToRead);
-	void ReadVector4ByIndex(FVectorVMContext& Context, FName AttributeToRead);
-	void ReadColorByIndex(FVectorVMContext& Context, FName AttributeToRead);
-	void ReadQuatByIndex(FVectorVMContext& Context, FName AttributeToRead);
-	void ReadIDByIndex(FVectorVMContext& Context, FName AttributeToRead);
+	void GetNumSpawnedParticles(FVectorVMExternalFunctionContext& Context);
+	void GetSpawnedIDAtIndex(FVectorVMExternalFunctionContext& Context);
+	void GetNumParticles(FVectorVMExternalFunctionContext& Context);
+	void GetParticleIndex(FVectorVMExternalFunctionContext& Context);
+	void GetParticleIndexFromIDTable(FVectorVMExternalFunctionContext& Context);
+	void ReadInt(FVectorVMExternalFunctionContext& Context, FName AttributeToRead);
+	void ReadBool(FVectorVMExternalFunctionContext& Context, FName AttributeToRead);
+	void ReadFloat(FVectorVMExternalFunctionContext& Context, FName AttributeToRead);
+	void ReadVector2(FVectorVMExternalFunctionContext& Context, FName AttributeToRead);
+	void ReadVector3(FVectorVMExternalFunctionContext& Context, FName AttributeToRead);
+	void ReadVector4(FVectorVMExternalFunctionContext& Context, FName AttributeToRead);
+	void ReadColor(FVectorVMExternalFunctionContext& Context, FName AttributeToRead);
+	void ReadQuat(FVectorVMExternalFunctionContext& Context, FName AttributeToRead);
+	void ReadID(FVectorVMExternalFunctionContext& Context, FName AttributeToRead);
+	void ReadIntByIndex(FVectorVMExternalFunctionContext& Context, FName AttributeToRead);
+	void ReadBoolByIndex(FVectorVMExternalFunctionContext& Context, FName AttributeToRead);
+	void ReadFloatByIndex(FVectorVMExternalFunctionContext& Context, FName AttributeToRead);
+	void ReadVector2ByIndex(FVectorVMExternalFunctionContext& Context, FName AttributeToRead);
+	void ReadVector3ByIndex(FVectorVMExternalFunctionContext& Context, FName AttributeToRead);
+	void ReadVector4ByIndex(FVectorVMExternalFunctionContext& Context, FName AttributeToRead);
+	void ReadColorByIndex(FVectorVMExternalFunctionContext& Context, FName AttributeToRead);
+	void ReadQuatByIndex(FVectorVMExternalFunctionContext& Context, FName AttributeToRead);
+	void ReadIDByIndex(FVectorVMExternalFunctionContext& Context, FName AttributeToRead);
 
 protected:
 	void GetPersistentIDFunctions(TArray<FNiagaraFunctionSignature>& OutFunctions);
diff --git a/Engine/Plugins/FX/Niagara/Source/Niagara/Classes/NiagaraDataInterfacePlatformSet.h b/Engine/Plugins/FX/Niagara/Source/Niagara/Classes/NiagaraDataInterfacePlatformSet.h
--- a/Engine/Plugins/FX/Niagara/Source/Niagara/Classes/NiagaraDataInterfacePlatformSet.h
+++ b/Engine/Plugins/FX/Niagara/Source/Niagara/Classes/NiagaraDataInterfacePlatformSet.h
@@ -35,6 +35,6 @@ protected:
 	//~ UNiagaraDataInterface interface
 
 public:
-	void IsActive(FVectorVMContext& Context);
+	void IsActive(FVectorVMExternalFunctionContext& Context);
 };
 
diff --git a/Engine/Plugins/FX/Niagara/Source/Niagara/Classes/NiagaraDataInterfaceRasterizationGrid3D.h b/Engine/Plugins/FX/Niagara/Source/Niagara/Classes/NiagaraDataInterfaceRasterizationGrid3D.h
--- a/Engine/Plugins/FX/Niagara/Source/Niagara/Classes/NiagaraDataInterfaceRasterizationGrid3D.h
+++ b/Engine/Plugins/FX/Niagara/Source/Niagara/Classes/NiagaraDataInterfaceRasterizationGrid3D.h
@@ -106,9 +106,9 @@ public:
 	virtual bool HasPreSimulateTick() const override { return true; }
 	//~ UNiagaraDataInterface interface END
 
-	void GetNumCells(FVectorVMContext& Context);
-	void SetNumCells(FVectorVMContext& Context);
-	void SetFloatResetValue(FVectorVMContext& Context);
+	void GetNumCells(FVectorVMExternalFunctionContext& Context);
+	void SetNumCells(FVectorVMExternalFunctionContext& Context);
+	void SetFloatResetValue(FVectorVMExternalFunctionContext& Context);
 
 	static const FName SetNumCellsFunctionName;
 	static const FName SetFloatResetValueFunctionName;	
diff --git a/Engine/Plugins/FX/Niagara/Source/Niagara/Classes/NiagaraDataInterfaceRenderTarget2D.h b/Engine/Plugins/FX/Niagara/Source/Niagara/Classes/NiagaraDataInterfaceRenderTarget2D.h
--- a/Engine/Plugins/FX/Niagara/Source/Niagara/Classes/NiagaraDataInterfaceRenderTarget2D.h
+++ b/Engine/Plugins/FX/Niagara/Source/Niagara/Classes/NiagaraDataInterfaceRenderTarget2D.h
@@ -121,8 +121,8 @@ public:
 	virtual bool RenderVariableToCanvas(FNiagaraSystemInstanceID SystemInstanceID, FName VariableName, class FCanvas* Canvas, const FIntRect& DrawRect) const override;
 	//~ UNiagaraDataInterface interface END
 
-	void GetSize(FVectorVMContext& Context); 
-	void SetSize(FVectorVMContext& Context);
+	void GetSize(FVectorVMExternalFunctionContext& Context); 
+	void SetSize(FVectorVMExternalFunctionContext& Context);
 
 	static const FName SetValueFunctionName;
 	static const FName GetValueFunctionName;
diff --git a/Engine/Plugins/FX/Niagara/Source/Niagara/Classes/NiagaraDataInterfaceRenderTarget2DArray.h b/Engine/Plugins/FX/Niagara/Source/Niagara/Classes/NiagaraDataInterfaceRenderTarget2DArray.h
--- a/Engine/Plugins/FX/Niagara/Source/Niagara/Classes/NiagaraDataInterfaceRenderTarget2DArray.h
+++ b/Engine/Plugins/FX/Niagara/Source/Niagara/Classes/NiagaraDataInterfaceRenderTarget2DArray.h
@@ -114,8 +114,8 @@ public:
 	virtual bool GetExposedVariableValue(const FNiagaraVariableBase& InVariable, void* InPerInstanceData, FNiagaraSystemInstance* InSystemInstance, void* OutData) const override;
 
 	//~ UNiagaraDataInterface interface END
-	void GetSize(FVectorVMContext& Context); 
-	void SetSize(FVectorVMContext& Context);
+	void GetSize(FVectorVMExternalFunctionContext& Context); 
+	void SetSize(FVectorVMExternalFunctionContext& Context);
 
 	static const FName SetValueFunctionName;
 	static const FName GetValueFunctionName;
diff --git a/Engine/Plugins/FX/Niagara/Source/Niagara/Classes/NiagaraDataInterfaceSimpleCounter.h b/Engine/Plugins/FX/Niagara/Source/Niagara/Classes/NiagaraDataInterfaceSimpleCounter.h
--- a/Engine/Plugins/FX/Niagara/Source/Niagara/Classes/NiagaraDataInterfaceSimpleCounter.h
+++ b/Engine/Plugins/FX/Niagara/Source/Niagara/Classes/NiagaraDataInterfaceSimpleCounter.h
@@ -48,7 +48,7 @@ public:
 	virtual void GetVMExternalFunction(const FVMExternalFunctionBindingInfo& BindingInfo, void* InstanceData, FVMExternalFunction &OutFunc) override;
 
 	// VM functions
-	void GetNextValue(FVectorVMContext& Context);
+	void GetNextValue(FVectorVMExternalFunctionContext& Context);
 
 	virtual bool Equals(const UNiagaraDataInterface* Other) const override;
 	virtual bool CanExecuteOnTarget(ENiagaraSimTarget Target) const override { return true; }
diff --git a/Engine/Plugins/FX/Niagara/Source/Niagara/Classes/NiagaraDataInterfaceSkeletalMesh.h b/Engine/Plugins/FX/Niagara/Source/Niagara/Classes/NiagaraDataInterfaceSkeletalMesh.h
--- a/Engine/Plugins/FX/Niagara/Source/Niagara/Classes/NiagaraDataInterfaceSkeletalMesh.h
+++ b/Engine/Plugins/FX/Niagara/Source/Niagara/Classes/NiagaraDataInterfaceSkeletalMesh.h
@@ -882,45 +882,45 @@ public:
 	void BindTriangleSamplingFunction(const FVMExternalFunctionBindingInfo& BindingInfo, FNDISkeletalMesh_InstanceData* InstData, FVMExternalFunction &OutFunc);
 
 	template<typename FilterMode, typename AreaWeightingMode>
-	void GetFilteredTriangleCount(FVectorVMContext& Context);
+	void GetFilteredTriangleCount(FVectorVMExternalFunctionContext& Context);
 
 	template<typename FilterMode, typename AreaWeightingMode>
-	void GetFilteredTriangleAt(FVectorVMContext& Context);
+	void GetFilteredTriangleAt(FVectorVMExternalFunctionContext& Context);
 
 	template<typename FilterMode, typename AreaWeightingMode>
-	void RandomTriCoord(FVectorVMContext& Context);
+	void RandomTriCoord(FVectorVMExternalFunctionContext& Context);
 
 	template<typename FilterMode, typename AreaWeightingMode>
-	void IsValidTriCoord(FVectorVMContext& Context);
+	void IsValidTriCoord(FVectorVMExternalFunctionContext& Context);
 
 	template<typename SkinningHandlerType, typename TransformHandlerType, typename VertexAccessorType, typename bInterpolated>
-	void GetTriCoordSkinnedData(FVectorVMContext& Context);
+	void GetTriCoordSkinnedData(FVectorVMExternalFunctionContext& Context);
 
 	template<typename TransformHandlerType, typename bInterpolated>
-	void GetTriCoordSkinnedDataFallback(FVectorVMContext& Context);
+	void GetTriCoordSkinnedDataFallback(FVectorVMExternalFunctionContext& Context);
 
-	void GetTriCoordColor(FVectorVMContext& Context);
+	void GetTriCoordColor(FVectorVMExternalFunctionContext& Context);
 
-	void GetTriCoordColorFallback(FVectorVMContext& Context);
+	void GetTriCoordColorFallback(FVectorVMExternalFunctionContext& Context);
 
 	template<typename VertexAccessorType>
-	void GetTriCoordUV(FVectorVMContext& Context);
+	void GetTriCoordUV(FVectorVMExternalFunctionContext& Context);
 
 	template<typename SkinningHandlerType>
-	void GetTriCoordVertices(FVectorVMContext& Context);
+	void GetTriCoordVertices(FVectorVMExternalFunctionContext& Context);
 
 	template<typename VertexAccessorType>
-	void GetTriangleCoordAtUV(FVectorVMContext& Context);
+	void GetTriangleCoordAtUV(FVectorVMExternalFunctionContext& Context);
 
 	template<typename VertexAccessorType>
-	void GetTriangleCoordInAabb(FVectorVMContext& Context);
+	void GetTriangleCoordInAabb(FVectorVMExternalFunctionContext& Context);
 
 private:
 	template<typename FilterMode, typename AreaWeightingMode>
 	FORCEINLINE int32 RandomTriIndex(FNDIRandomHelper& RandHelper, FSkeletalMeshAccessorHelper& Accessor, FNDISkeletalMesh_InstanceData* InstData, int32 InstanceIndex);
 
-	void RandomTriangle(FVectorVMContext& Context);
-	void GetTriangleCount(FVectorVMContext& Context);
+	void RandomTriangle(FVectorVMExternalFunctionContext& Context);
+	void GetTriangleCount(FVectorVMExternalFunctionContext& Context);
 
 	template<typename FilterMode, typename AreaWeightingMode>
 	FORCEINLINE int32 GetFilteredTriangleCount(FSkeletalMeshAccessorHelper& Accessor, FNDISkeletalMesh_InstanceData* InstData);
@@ -938,28 +938,28 @@ public:
 	void GetVertexSamplingFunctions(TArray<FNiagaraFunctionSignature>& OutFunctions);
 	void BindVertexSamplingFunction(const FVMExternalFunctionBindingInfo& BindingInfo, FNDISkeletalMesh_InstanceData* InstData, FVMExternalFunction &OutFunc);
 
-	void IsValidVertex(FVectorVMContext& Context);
-	void RandomVertex(FVectorVMContext& Context);
-	void GetVertexCount(FVectorVMContext& Context);
+	void IsValidVertex(FVectorVMExternalFunctionContext& Context);
+	void RandomVertex(FVectorVMExternalFunctionContext& Context);
+	void GetVertexCount(FVectorVMExternalFunctionContext& Context);
 
 	template<typename FilterMode>
-	void IsValidFilteredVertex(FVectorVMContext& Context);
+	void IsValidFilteredVertex(FVectorVMExternalFunctionContext& Context);
 	template<typename FilterMode>
-	void RandomFilteredVertex(FVectorVMContext& Context);
+	void RandomFilteredVertex(FVectorVMExternalFunctionContext& Context);
 	template<typename FilterMode>
-	void GetFilteredVertexCount(FVectorVMContext& Context);
+	void GetFilteredVertexCount(FVectorVMExternalFunctionContext& Context);
 	template<typename FilterMode>
-	void GetFilteredVertexAt(FVectorVMContext& Context);
+	void GetFilteredVertexAt(FVectorVMExternalFunctionContext& Context);
 
 	template<typename SkinningHandlerType, typename TransformHandlerType, typename VertexAccessorType>
-	void GetVertexSkinnedData(FVectorVMContext& Context);
+	void GetVertexSkinnedData(FVectorVMExternalFunctionContext& Context);
 
-	void GetVertexColor(FVectorVMContext& Context);
+	void GetVertexColor(FVectorVMExternalFunctionContext& Context);
 
-	void GetVertexColorFallback(FVectorVMContext& Context);
+	void GetVertexColorFallback(FVectorVMExternalFunctionContext& Context);
 
 	template<typename VertexAccessorType>
-	void GetVertexUV(FVectorVMContext& Context);
+	void GetVertexUV(FVectorVMExternalFunctionContext& Context);
 
 private:
 	template<typename FilterMode>
@@ -982,28 +982,28 @@ public:
 	void BindSkeletonSamplingFunction(const FVMExternalFunctionBindingInfo& BindingInfo, FNDISkeletalMesh_InstanceData* InstData, FVMExternalFunction &OutFunc);
 
 	template<typename SkinningHandlerType, typename TransformHandlerType, typename bInterpolated>
-	void GetSkinnedBoneData(FVectorVMContext& Context);
+	void GetSkinnedBoneData(FVectorVMExternalFunctionContext& Context);
 
-	void IsValidBone(FVectorVMContext& Context);
-	void RandomBone(FVectorVMContext& Context);
-	void GetBoneCount(FVectorVMContext& Context);
+	void IsValidBone(FVectorVMExternalFunctionContext& Context);
+	void RandomBone(FVectorVMExternalFunctionContext& Context);
+	void GetBoneCount(FVectorVMExternalFunctionContext& Context);
 
-	void GetFilteredBoneCount(FVectorVMContext& Context);
-	void GetFilteredBoneAt(FVectorVMContext& Context);
-	void RandomFilteredBone(FVectorVMContext& Context);
+	void GetFilteredBoneCount(FVectorVMExternalFunctionContext& Context);
+	void GetFilteredBoneAt(FVectorVMExternalFunctionContext& Context);
+	void RandomFilteredBone(FVectorVMExternalFunctionContext& Context);
 
-	void GetUnfilteredBoneCount(FVectorVMContext& Context);
-	void GetUnfilteredBoneAt(FVectorVMContext& Context);
-	void RandomUnfilteredBone(FVectorVMContext& Context);
+	void GetUnfilteredBoneCount(FVectorVMExternalFunctionContext& Context);
+	void GetUnfilteredBoneAt(FVectorVMExternalFunctionContext& Context);
+	void RandomUnfilteredBone(FVectorVMExternalFunctionContext& Context);
 
-	void GetFilteredSocketCount(FVectorVMContext& Context);
-	void GetFilteredSocketBoneAt(FVectorVMContext& Context);
-	void GetFilteredSocketTransform(FVectorVMContext& Context);
-	void RandomFilteredSocket(FVectorVMContext& Context);
+	void GetFilteredSocketCount(FVectorVMExternalFunctionContext& Context);
+	void GetFilteredSocketBoneAt(FVectorVMExternalFunctionContext& Context);
+	void GetFilteredSocketTransform(FVectorVMExternalFunctionContext& Context);
+	void RandomFilteredSocket(FVectorVMExternalFunctionContext& Context);
 
-	void RandomFilteredSocketOrBone(FVectorVMContext& Context);
-	void GetFilteredSocketOrBoneCount(FVectorVMContext& Context);
-	void GetFilteredSocketOrBoneBoneAt(FVectorVMContext& Context);
+	void RandomFilteredSocketOrBone(FVectorVMExternalFunctionContext& Context);
+	void GetFilteredSocketOrBoneCount(FVectorVMExternalFunctionContext& Context);
+	void GetFilteredSocketOrBoneBoneAt(FVectorVMExternalFunctionContext& Context);
 	// End of Direct Bone + Socket Sampling
 	//////////////////////////////////////////////////////////////////////////
 
diff --git a/Engine/Plugins/FX/Niagara/Source/Niagara/Classes/NiagaraDataInterfaceSpline.h b/Engine/Plugins/FX/Niagara/Source/Niagara/Classes/NiagaraDataInterfaceSpline.h
--- a/Engine/Plugins/FX/Niagara/Source/Niagara/Classes/NiagaraDataInterfaceSpline.h
+++ b/Engine/Plugins/FX/Niagara/Source/Niagara/Classes/NiagaraDataInterfaceSpline.h
@@ -86,22 +86,22 @@ public:
 	//UNiagaraDataInterface Interface End
 
 	template<typename TransformHandlerType, typename SplineSampleType>
-	void SampleSplinePositionByUnitDistance(FVectorVMContext& Context);
+	void SampleSplinePositionByUnitDistance(FVectorVMExternalFunctionContext& Context);
 	template<typename TransformHandlerType, typename SplineSampleType>
-	void SampleSplineRotationByUnitDistance(FVectorVMContext& Context);
+	void SampleSplineRotationByUnitDistance(FVectorVMExternalFunctionContext& Context);
 	template<typename TransformHandlerType, typename SplineSampleType>
-	void SampleSplineUpVectorByUnitDistance(FVectorVMContext& Context);
+	void SampleSplineUpVectorByUnitDistance(FVectorVMExternalFunctionContext& Context);
 	template<typename TransformHandlerType, typename SplineSampleType>
-	void SampleSplineRightVectorByUnitDistance(FVectorVMContext& Context);
+	void SampleSplineRightVectorByUnitDistance(FVectorVMExternalFunctionContext& Context);
 	template<typename TransformHandlerType, typename SplineSampleType>
-	void SampleSplineDirectionByUnitDistance(FVectorVMContext& Context);
+	void SampleSplineDirectionByUnitDistance(FVectorVMExternalFunctionContext& Context);
 	template<typename TransformHandlerType, typename SplineSampleType>
-	void SampleSplineTangentByUnitDistance(FVectorVMContext& Context);
+	void SampleSplineTangentByUnitDistance(FVectorVMExternalFunctionContext& Context);
 	template<typename PosXType, typename PosYType, typename PosZType>
-	void FindClosestUnitDistanceFromPositionWS(FVectorVMContext& Context);
+	void FindClosestUnitDistanceFromPositionWS(FVectorVMExternalFunctionContext& Context);
 	
-	void GetLocalToWorld(FVectorVMContext& Context);
-	void GetLocalToWorldInverseTransposed(FVectorVMContext& Context);
+	void GetLocalToWorld(FVectorVMExternalFunctionContext& Context);
+	void GetLocalToWorldInverseTransposed(FVectorVMExternalFunctionContext& Context);
 
 	virtual int32 PerInstanceDataPassedToRenderThreadSize() const override
 	{
@@ -113,6 +113,6 @@ protected:
 
 private:
 	
-	void WriteTransform(const FMatrix& ToWrite, FVectorVMContext& Context);
+	void WriteTransform(const FMatrix& ToWrite, FVectorVMExternalFunctionContext& Context);
 	//Cached ComponentToWorld.
 };
diff --git a/Engine/Plugins/FX/Niagara/Source/Niagara/Classes/NiagaraDataInterfaceStaticMesh.h b/Engine/Plugins/FX/Niagara/Source/Niagara/Classes/NiagaraDataInterfaceStaticMesh.h
--- a/Engine/Plugins/FX/Niagara/Source/Niagara/Classes/NiagaraDataInterfaceStaticMesh.h
+++ b/Engine/Plugins/FX/Niagara/Source/Niagara/Classes/NiagaraDataInterfaceStaticMesh.h
@@ -344,52 +344,52 @@ public:
 public:
 	UStaticMesh* GetStaticMesh(TWeakObjectPtr<USceneComponent>& OutComponent, class FNiagaraSystemInstance* SystemInstance = nullptr);
 
-	void IsValid(FVectorVMContext& Context);
+	void IsValid(FVectorVMExternalFunctionContext& Context);
 
 	template<typename TSampleMode>
-	void RandomSection(FVectorVMContext& Context);
+	void RandomSection(FVectorVMExternalFunctionContext& Context);
 
 	template<typename TSampleMode>
-	void RandomTriCoord(FVectorVMContext& Context);
+	void RandomTriCoord(FVectorVMExternalFunctionContext& Context);
 
 	template<typename TSampleMode>
-	void RandomTriCoordOnSection(FVectorVMContext& Context);
+	void RandomTriCoordOnSection(FVectorVMExternalFunctionContext& Context);
 
- 	void RandomTriCoordVertexColorFiltered(FVectorVMContext& Context);
+ 	void RandomTriCoordVertexColorFiltered(FVectorVMExternalFunctionContext& Context);
 	
 	template<typename TransformHandlerType>
-	void GetTriCoordPosition(FVectorVMContext& Context);
+	void GetTriCoordPosition(FVectorVMExternalFunctionContext& Context);
 
 	template<typename TransformHandlerType>
-	void GetTriCoordNormal(FVectorVMContext& Context);
+	void GetTriCoordNormal(FVectorVMExternalFunctionContext& Context);
 
 	template<typename VertexAccessorType, typename TransformHandlerType>
-	void GetTriCoordTangents(FVectorVMContext& Context);
+	void GetTriCoordTangents(FVectorVMExternalFunctionContext& Context);
 
-	void GetTriCoordColor(FVectorVMContext& Context);
+	void GetTriCoordColor(FVectorVMExternalFunctionContext& Context);
 
 	template<typename VertexAccessorType>
-	void GetTriCoordUV(FVectorVMContext& Context);
+	void GetTriCoordUV(FVectorVMExternalFunctionContext& Context);
 
-	void GetTriCoordPositionAndVelocity(FVectorVMContext& Context);
+	void GetTriCoordPositionAndVelocity(FVectorVMExternalFunctionContext& Context);
 
-	void GetLocalToWorld(FVectorVMContext& Context);
-	void GetLocalToWorldInverseTransposed(FVectorVMContext& Context);
-	void GetWorldVelocity(FVectorVMContext& Context);
+	void GetLocalToWorld(FVectorVMExternalFunctionContext& Context);
+	void GetLocalToWorldInverseTransposed(FVectorVMExternalFunctionContext& Context);
+	void GetWorldVelocity(FVectorVMExternalFunctionContext& Context);
 
 	template<typename TransformHandlerType>
-	void GetVertexPosition(FVectorVMContext& Context);
+	void GetVertexPosition(FVectorVMExternalFunctionContext& Context);
 
 	// Socket Functions
-	void GetSocketCount(FVectorVMContext& Context);
-	void GetFilteredSocketCount(FVectorVMContext& Context);
-	void GetUnfilteredSocketCount(FVectorVMContext& Context);
+	void GetSocketCount(FVectorVMExternalFunctionContext& Context);
+	void GetFilteredSocketCount(FVectorVMExternalFunctionContext& Context);
+	void GetUnfilteredSocketCount(FVectorVMExternalFunctionContext& Context);
 	template<bool bWorldSpace>
-	void GetSocketTransform(FVectorVMContext& Context);
+	void GetSocketTransform(FVectorVMExternalFunctionContext& Context);
 	template<bool bWorldSpace>
-	void GetFilteredSocketTransform(FVectorVMContext& Context);
+	void GetFilteredSocketTransform(FVectorVMExternalFunctionContext& Context);
 	template<bool bWorldSpace>
-	void GetUnfilteredSocketTransform(FVectorVMContext& Context);
+	void GetUnfilteredSocketTransform(FVectorVMExternalFunctionContext& Context);
 
 	void SetSourceComponentFromBlueprints(UStaticMeshComponent* ComponentToUse);
 	void SetDefaultMeshFromBlueprints(UStaticMesh* MeshToUse);
@@ -413,7 +413,7 @@ private:
 	template<typename TSampleMode>
 	FORCEINLINE_DEBUGGABLE int32 RandomTriIndexOnSection(FRandomStream& RandStream, const FStaticMeshLODResources& Res, int32 SectionIdx, FNDIStaticMesh_InstanceData* InstData);
 
-	void WriteTransform(const FMatrix& ToWrite, FVectorVMContext& Context);
+	void WriteTransform(const FMatrix& ToWrite, FVectorVMExternalFunctionContext& Context);
 };
 
 //TODO: IMO this should be generalized fuhrer if possible and extended to a system allowing filtering based on texture color etc too.
