geometry collection di

diff --git a/Engine/Plugins/Experimental/ChaosNiagara/Shaders/NiagaraDataInterfaceGeometryCollection.ush b/Engine/Plugins/Experimental/ChaosNiagara/Shaders/NiagaraDataInterfaceGeometryCollection.ush
new file mode 100644
--- /dev/null
+++ b/Engine/Plugins/Experimental/ChaosNiagara/Shaders/NiagaraDataInterfaceGeometryCollection.ush
@@ -0,0 +1,214 @@
+// Copyright Epic Games, Inc. All Rights Reserved.
+
+/*=============================================================================
+NiagaraDataInterfaceGeometryCollection.ush 
+=============================================================================*/
+
+#pragma once
+
+/* -----------------------------------------------------------------
+ * static mesh constants and context
+ * -----------------------------------------------------------------
+ */
+
+#define SMALL_NUMBER 1e-8
+#define MAX_DISTANCE 1e+38
+#define BOX_INDEX 0
+#define SPHERE_INDEX 1
+#define CAPSULE_INDEX 2
+
+struct FDIGeometryCollectionContext
+{
+	uint NumPieces;	
+	Buffer<float4>		WorldTransformBuffer;
+	Buffer<float4>		PrevWorldTransformBuffer;
+	Buffer<float4>		WorldInverseTransformBuffer;
+	Buffer<float4>		PrevWorldInverseTransformBuffer;
+	Buffer<float4>		BoundsBuffer;
+}; 
+
+#define DIGEOMETRYCOLLECTION_DECLARE_CONSTANTS(NAME)\
+uint NumPieces_##NAME;\
+Buffer<float4>		WorldTransformBuffer_##NAME;\
+Buffer<float4>		PrevWorldTransformBuffer_##NAME;\
+Buffer<float4>		WorldInverseTransformBuffer_##NAME;\
+Buffer<float4>		PrevWorldInverseTransformBuffer_##NAME;\
+Buffer<float4>		BoundsBuffer_##NAME;\
+
+
+#define DIGEOMETRYCOLLECTION_MAKE_CONTEXT(NAME)\
+FDIGeometryCollectionContext DIContext; \
+DIContext.NumPieces = NumPieces_##NAME;\
+DIContext.WorldTransformBuffer =  WorldTransformBuffer_##NAME;\
+DIContext.PrevWorldTransformBuffer = PrevWorldTransformBuffer_##NAME;\
+DIContext.WorldInverseTransformBuffer = WorldInverseTransformBuffer_##NAME;\
+DIContext.PrevWorldInverseTransformBuffer = PrevWorldInverseTransformBuffer_##NAME;\
+DIContext.BoundsBuffer = BoundsBuffer_##NAME;\
+
+
+
+/* -----------------------------------------------------------------
+ * Position / Transform utils
+ * -----------------------------------------------------------------
+ */
+
+float3 DIGeometryCollection_GetCurrentPosition(in FDIGeometryCollectionContext DIContext, in int ElementIndex)
+{
+	const int BufferOffset = 3 * ElementIndex;
+    return float3(DIContext.WorldTransformBuffer[BufferOffset].w, DIContext.WorldTransformBuffer[BufferOffset+1].w, DIContext.WorldTransformBuffer[BufferOffset+2].w);
+}
+
+float3 DIGeometryCollection_GetPreviousPosition(in FDIGeometryCollectionContext DIContext, in int ElementIndex)
+{
+	const int BufferOffset = 3 * ElementIndex;
+    return float3(DIContext.WorldTransformBuffer[BufferOffset].w, DIContext.WorldTransformBuffer[BufferOffset+1].w, DIContext.WorldTransformBuffer[BufferOffset+2].w);
+}
+
+float3x4 DIGeometryCollection_GetCurrentTransform(in FDIGeometryCollectionContext DIContext, in int ElementIndex)
+{
+	const int BufferOffset = 3 * ElementIndex;
+    return float3x4(DIContext.WorldTransformBuffer[BufferOffset], DIContext.WorldTransformBuffer[BufferOffset+1], DIContext.WorldTransformBuffer[BufferOffset+2]);
+}
+
+float3x4 DIGeometryCollection_GetInverseTransform(in FDIGeometryCollectionContext DIContext, in int ElementIndex)
+{
+	const int BufferOffset = 3 * ElementIndex;
+    return float3x4(DIContext.WorldInverseTransformBuffer[BufferOffset], DIContext.WorldInverseTransformBuffer[BufferOffset+1], DIContext.WorldInverseTransformBuffer[BufferOffset+2]);
+}
+
+float3x4 DIGeometryCollection_GetPreviousTransform(in FDIGeometryCollectionContext DIContext, in int ElementIndex)
+{
+	const int BufferOffset = 3 * ElementIndex;
+    return float3x4(DIContext.PrevWorldTransformBuffer[BufferOffset], DIContext.PrevWorldTransformBuffer[BufferOffset+1], DIContext.PrevWorldTransformBuffer[BufferOffset+2]);
+}
+
+float3x4 DIGeometryCollection_GetPreviousInverse(in FDIGeometryCollectionContext DIContext, in int ElementIndex)
+{
+	const int BufferOffset = 3 * ElementIndex;
+    return float3x4(DIContext.PrevWorldInverseTransformBuffer[BufferOffset], DIContext.PrevWorldInverseTransformBuffer[BufferOffset+1], DIContext.PrevWorldInverseTransformBuffer[BufferOffset+2]);
+}
+
+
+/* -----------------------------------------------------------------
+ * Collision detection utils
+ * -----------------------------------------------------------------
+ */
+
+// Given a world space position (WorldPosition) compute the sphere closest point (position,normal,velocity)
+float DIGeometryCollection_GetSphereProjection(in float3 LocalPosition, in float3 SphereCenter, in float SphereRadius, in int SphereIndex,
+			inout float3 OutClosestPosition, inout float3 OutClosestNormal, inout int OutElementIndex, inout float OutMinDistance)
+{
+	const float3 DeltaPosition = LocalPosition - SphereCenter;
+	const float DeltaLength = length(DeltaPosition);
+	const float SphereDistance = DeltaLength - SphereRadius;
+
+	if (SphereDistance < OutMinDistance)
+	{
+		OutMinDistance = SphereDistance;
+		OutElementIndex = SphereIndex;
+
+		OutClosestNormal = (DeltaLength > SMALL_NUMBER) ? DeltaPosition / DeltaLength : float3(0,0,0);
+		OutClosestPosition = LocalPosition - OutClosestNormal*SphereDistance;
+	}
+	return SphereDistance;
+}
+
+// Given a world space position (WorldPosition) compute the sphere closest point (position,normal,velocity)
+float DIGeometryCollection_GetCapsuleProjection(in float3 LocalPosition, in float2 CapsuleSize, in int CapsuleIndex,
+			inout float3 OutClosestPosition, inout float3 OutClosestNormal, inout int OutElementIndex, inout float OutMinDistance)
+{
+	const float HalfLength = 0.5 * CapsuleSize.y;
+	const float3 SphereCenter = float3(0,0,clamp(LocalPosition.z, -HalfLength, HalfLength));
+
+	return DIGeometryCollection_GetSphereProjection(LocalPosition,SphereCenter,CapsuleSize.x,CapsuleIndex,
+			OutClosestPosition,OutClosestNormal,OutElementIndex,OutMinDistance);
+}
+
+// Given a world space position (WorldPosition) compute the sphere closest point (position,normal,velocity)
+float DIGeometryCollection_GetBoxProjection(in float3 LocalPosition, in float3 BoxExtent, in int BoxIndex, 
+			inout float3 OutClosestPosition, inout float3 OutClosestNormal, inout int OutElementIndex, inout float OutMinDistance)
+{
+	const float3 HalfExtent = 0.5 * BoxExtent;
+	const float3 DeltaPosition = abs(LocalPosition) - HalfExtent;
+
+	const int ClosestAxis = ((DeltaPosition.x > DeltaPosition.y) && (DeltaPosition.x > DeltaPosition.z)) ? 0 : ( DeltaPosition.y > DeltaPosition.z) ? 1 : 2;
+	const float OutsideDistance = length(max(DeltaPosition,0.0));
+	const float BoxDistance = OutsideDistance + min(DeltaPosition[ClosestAxis],0.0);
+
+	if (BoxDistance < OutMinDistance)
+	{
+		OutMinDistance = BoxDistance;
+		OutElementIndex = BoxIndex;
+
+		if (BoxDistance <= 0)
+		{
+			const bool NegativeSide = LocalPosition[ClosestAxis] < 0.0;
+
+			OutClosestPosition = LocalPosition;
+			OutClosestNormal = float3(0,0,0);
+			if( ClosestAxis == 0) 
+			{
+				OutClosestPosition.x = NegativeSide ? -HalfExtent.x : HalfExtent.x;
+				OutClosestNormal.x = NegativeSide ? -1.0 : 1.0;
+			}
+			else if( ClosestAxis == 1) 
+			{
+				OutClosestPosition.y = NegativeSide ? -HalfExtent.y : HalfExtent.y;
+				OutClosestNormal.y = NegativeSide ? -1.0 : 1.0;
+			}
+			else if( ClosestAxis == 2) 
+			{
+				OutClosestPosition.z = NegativeSide ? -HalfExtent.z : HalfExtent.z;
+				OutClosestNormal.z = NegativeSide ? -1.0 : 1.0;
+			}
+		}
+		else
+		{
+			OutClosestPosition = clamp(LocalPosition,-HalfExtent,HalfExtent);
+			OutClosestNormal = (LocalPosition - OutClosestPosition) / OutsideDistance;
+		}
+	}
+	return BoxDistance;
+}
+
+
+/* -----------------------------------------------------------------
+ * Get the closest element to the world position
+ * -----------------------------------------------------------------
+ */
+
+float3 DIGeometryCollection_GetLocalPosition(in FDIGeometryCollectionContext DIContext, in float3 WorldPosition, in int ElementIndex, in float TimeFraction )
+{
+	const float3 CurrentLocal = mul(DIGeometryCollection_GetInverseTransform(DIContext,ElementIndex), float4(WorldPosition,1.0)).xyz;
+	const float3 PreviousLocal = mul(DIGeometryCollection_GetPreviousInverse(DIContext,ElementIndex), float4(WorldPosition,1.0)).xyz;
+
+	return PreviousLocal + TimeFraction * (CurrentLocal-PreviousLocal);
+}
+
+// Given a world space position (WorldPosition) compute the static mesh closest point (position,normal,velocity)
+void DIGeometryCollection_GetClosestPointNoNormal(in FDIGeometryCollectionContext DIContext, in float3 WorldPosition, in float DeltaTime, in float TimeFraction, out float OutClosestDistance, out float3 OutClosestPosition, out float3 OutClosestVelocity)
+{	
+	float3 CollisionPosition = float3(0,0,0);
+	float3 CollisionNormal = float3(0,0,0);
+	OutClosestDistance = MAX_DISTANCE;	
+	
+	OutClosestPosition = float3(0,0,0);
+	OutClosestVelocity = float3(0,0,0);
+	
+	int ClosestTransform = -1;
+	for (int CurrTransform = 0; CurrTransform < DIContext.NumPieces; ++CurrTransform)
+	{
+		float3 LocalPosition = DIGeometryCollection_GetLocalPosition(DIContext,WorldPosition,CurrTransform,TimeFraction);
+		
+		DIGeometryCollection_GetBoxProjection(LocalPosition, DIContext.BoundsBuffer[CurrTransform].xyz, CurrTransform,
+				CollisionPosition, CollisionNormal, ClosestTransform, OutClosestDistance);
+	}
+
+	if (ClosestTransform != -1)
+	{
+		const float3 PreviousPosition = mul(DIGeometryCollection_GetPreviousTransform(DIContext,ClosestTransform), float4(CollisionPosition,1.0)).xyz;
+		const float3 CurrentPosition = mul(DIGeometryCollection_GetCurrentTransform(DIContext,ClosestTransform), float4(CollisionPosition,1.0)).xyz;
+		OutClosestVelocity = ( CurrentPosition - PreviousPosition ) / DeltaTime;
+		OutClosestPosition = PreviousPosition + TimeFraction * (CurrentPosition-PreviousPosition);		
+	}
+}
\ No newline at end of file
diff --git a/Engine/Plugins/Experimental/ChaosNiagara/Source/ChaosNiagara/Private/NiagaraDataInterfaceGeometryCollection.cpp b/Engine/Plugins/Experimental/ChaosNiagara/Source/ChaosNiagara/Private/NiagaraDataInterfaceGeometryCollection.cpp
new file mode 100644
--- /dev/null
+++ b/Engine/Plugins/Experimental/ChaosNiagara/Source/ChaosNiagara/Private/NiagaraDataInterfaceGeometryCollection.cpp
@@ -0,0 +1,581 @@
+// Copyright Epic Games, Inc. All Rights Reserved.
+
+#include "NiagaraDataInterfaceGeometryCollection.h"
+#include "Animation/SkeletalMeshActor.h"
+#include "Components/SkeletalMeshComponent.h"
+#include "SkeletalRenderPublic.h"
+#include "SkeletalMeshTypes.h"
+#include "AnimationRuntime.h"
+#include "NiagaraShader.h"
+#include "NiagaraComponent.h"
+#include "NiagaraRenderer.h"
+#include "NiagaraSimStageData.h"
+#include "NiagaraSystemInstance.h"
+#include "ShaderParameterUtils.h"
+#include "EngineUtils.h"
+#include "Renderer/Private/ScenePrivate.h"
+#include "DistanceFieldAtlas.h"
+#include "Renderer/Private/DistanceFieldLightingShared.h"
+#include "NiagaraGpuComputeDispatchInterface.h"
+#include "GeometryCollection//GeometryCollectionComponent.h"
+
+#define LOCTEXT_NAMESPACE "NiagaraDataInterfaceGeometryCollection"
+DEFINE_LOG_CATEGORY_STATIC(LogGeometryCollection, Log, All);
+
+//------------------------------------------------------------------------------------------------------------
+
+static const FName GetClosestPointNoNormalName(TEXT("GetClosestPointNoNormal"));
+
+
+//------------------------------------------------------------------------------------------------------------
+
+const FString UNiagaraDataInterfaceGeometryCollection::NumPiecesName(TEXT("NumPieces_"));
+const FString UNiagaraDataInterfaceGeometryCollection::WorldTransformBufferName(TEXT("WorldTransformBuffer_"));
+const FString UNiagaraDataInterfaceGeometryCollection::PrevWorldTransformBufferName(TEXT("PrevWorldTransformBuffer_"));
+const FString UNiagaraDataInterfaceGeometryCollection::WorldInverseTransformBufferName(TEXT("WorldInverseTransformBuffer_"));
+const FString UNiagaraDataInterfaceGeometryCollection::PrevWorldInverseTransformBufferName(TEXT("PrevWorldInverseTransformBuffer_"));
+const FString UNiagaraDataInterfaceGeometryCollection::BoundsBufferName(TEXT("BoundsBuffer_"));
+
+
+//------------------------------------------------------------------------------------------------------------
+
+struct FNDIGeometryCollectionParametersName
+{
+	FNDIGeometryCollectionParametersName(const FString& Suffix)
+	{
+		NumPiecesName = UNiagaraDataInterfaceGeometryCollection::NumPiecesName + Suffix;
+		WorldTransformBufferName = UNiagaraDataInterfaceGeometryCollection::WorldTransformBufferName + Suffix;
+		PrevWorldTransformBufferName = UNiagaraDataInterfaceGeometryCollection::PrevWorldTransformBufferName + Suffix;
+		WorldInverseTransformBufferName = UNiagaraDataInterfaceGeometryCollection::WorldInverseTransformBufferName + Suffix;
+		PrevWorldInverseTransformBufferName = UNiagaraDataInterfaceGeometryCollection::PrevWorldInverseTransformBufferName + Suffix;
+		BoundsBufferName = UNiagaraDataInterfaceGeometryCollection::BoundsBufferName + Suffix;
+	}
+
+	FString NumPiecesName;
+	FString WorldTransformBufferName;
+	FString PrevWorldTransformBufferName;
+	FString WorldInverseTransformBufferName;
+	FString PrevWorldInverseTransformBufferName;
+	FString BoundsBufferName;
+};
+
+//------------------------------------------------------------------------------------------------------------
+
+template<typename BufferType, EPixelFormat PixelFormat>
+void CreateInternalBuffer(FRWBuffer& OutputBuffer, uint32 ElementCount)
+{
+	if (ElementCount > 0)
+	{
+		OutputBuffer.Initialize(TEXT("FNDIGeometryCollectionBuffer"), sizeof(BufferType), ElementCount, PixelFormat, BUF_Static);
+	}
+}
+
+template<typename BufferType, EPixelFormat PixelFormat>
+void UpdateInternalBuffer(const TArray<BufferType>& InputData, FRWBuffer& OutputBuffer)
+{
+	uint32 ElementCount = InputData.Num();
+	if (ElementCount > 0 && OutputBuffer.Buffer.IsValid())
+	{
+		const uint32 BufferBytes = sizeof(BufferType) * ElementCount;
+
+		void* OutputData = RHILockBuffer(OutputBuffer.Buffer, 0, BufferBytes, RLM_WriteOnly);
+
+		FMemory::Memcpy(OutputData, InputData.GetData(), BufferBytes);
+		RHIUnlockBuffer(OutputBuffer.Buffer);
+	}
+}
+
+
+//------------------------------------------------------------------------------------------------------------
+
+void FNDIGeometryCollectionBuffer::InitRHI()
+{
+	CreateInternalBuffer<FVector4f, EPixelFormat::PF_A32B32G32R32F>(WorldTransformBuffer, 3 * NumPieces);
+	CreateInternalBuffer<FVector4f, EPixelFormat::PF_A32B32G32R32F>(PrevWorldTransformBuffer, 3 * NumPieces);
+
+	CreateInternalBuffer<FVector4f, EPixelFormat::PF_A32B32G32R32F>(WorldInverseTransformBuffer, 3 * NumPieces);
+	CreateInternalBuffer<FVector4f, EPixelFormat::PF_A32B32G32R32F>(PrevWorldInverseTransformBuffer, 3 * NumPieces);
+
+	CreateInternalBuffer<FVector4f, EPixelFormat::PF_A32B32G32R32F>(BoundsBuffer, NumPieces);
+}
+
+void FNDIGeometryCollectionBuffer::ReleaseRHI()
+{
+	WorldTransformBuffer.Release();
+	PrevWorldTransformBuffer.Release();
+	WorldInverseTransformBuffer.Release();
+	PrevWorldInverseTransformBuffer.Release();
+	BoundsBuffer.Release();
+}
+
+//------------------------------------------------------------------------------------------------------------
+
+
+
+void FNDIGeometryCollectionData::Release()
+{
+	if (AssetBuffer)
+	{
+		BeginReleaseResource(AssetBuffer);
+		ENQUEUE_RENDER_COMMAND(DeleteResource)(
+			[ParamPointerToRelease = AssetBuffer](FRHICommandListImmediate& RHICmdList)
+			{
+				delete ParamPointerToRelease;
+			});
+		AssetBuffer = nullptr;
+	}
+}
+
+void FNDIGeometryCollectionData::Init(UNiagaraDataInterfaceGeometryCollection* Interface, FNiagaraSystemInstance* SystemInstance)
+{
+	AssetBuffer = nullptr;
+
+	if (Interface != nullptr && SystemInstance != nullptr)
+	{		
+		if (Interface->GeometryCollectionActor != nullptr && Interface->GeometryCollectionActor->GetGeometryCollectionComponent() != nullptr)
+		{			
+			const TSharedPtr<FGeometryCollection, ESPMode::ThreadSafe>
+				Collection = Interface->GeometryCollectionActor->GetGeometryCollectionComponent()->RestCollection->GetGeometryCollection();
+			const TManagedArray<FBox>& BoundingBoxes = Collection->BoundingBox;
+			const TManagedArray<int32>& TransformIndex = Collection->TransformIndex;
+			const TManagedArray<FTransform>& Transforms = Interface->GeometryCollectionActor->GetGeometryCollectionComponent()->GetTransformArray();
+			int NumPieces = BoundingBoxes.Num();
+
+			AssetArrays = new FNDIGeometryCollectionArrays();
+			AssetArrays->Resize(NumPieces);
+
+			AssetBuffer = new FNDIGeometryCollectionBuffer();
+			AssetBuffer->SetNumPieces(NumPieces);
+			BeginInitResource(AssetBuffer);
+			
+			for (int i = 0; i < NumPieces; ++i)
+			{
+				FBox CurrBox = BoundingBoxes[i];
+				FVector3f BoxSize = CurrBox.Max - CurrBox.Min;
+				AssetArrays->BoundsBuffer[i] = FVector4f(BoxSize.X, BoxSize.Y, BoxSize.Z, 0);
+			}
+		}
+		else
+		{
+			AssetArrays = new FNDIGeometryCollectionArrays();
+			AssetArrays->Resize(1);
+
+			AssetBuffer = new FNDIGeometryCollectionBuffer();
+			AssetBuffer->SetNumPieces(1);
+			BeginInitResource(AssetBuffer);
+		}
+	}
+
+}
+
+void FNDIGeometryCollectionData::Update(UNiagaraDataInterfaceGeometryCollection* Interface, FNiagaraSystemInstance* SystemInstance)
+{
+	if (Interface != nullptr && SystemInstance != nullptr)
+	{		
+		TickingGroup = ComputeTickingGroup();
+		
+		if (Interface->GeometryCollectionActor != nullptr && Interface->GeometryCollectionActor->GetGeometryCollectionComponent() != nullptr &&
+			Interface->GeometryCollectionActor->GetGeometryCollectionComponent()->RestCollection != nullptr)
+		{			
+			const FTransform ActorTransform = Interface->GeometryCollectionActor->GetTransform();
+
+			const TSharedPtr<FGeometryCollection, ESPMode::ThreadSafe>
+				Collection = Interface->GeometryCollectionActor->GetGeometryCollectionComponent()->RestCollection->GetGeometryCollection();
+			const TManagedArray<FBox>& BoundingBoxes = Collection->BoundingBox;
+			const TManagedArray<int32>& TransformIndexArray = Collection->TransformIndex;
+			const TManagedArray<FTransform>& Transforms = Interface->GeometryCollectionActor->GetGeometryCollectionComponent()->GetTransformArray();
+			const TArray<FMatrix>& GlobalMatrices = Interface->GeometryCollectionActor->GetGeometryCollectionComponent()->GetGlobalMatrices();
+
+			// #todo(dmp): rename transform to pieces or something
+			int32 NumPieces = BoundingBoxes.Num();
+
+			if (GlobalMatrices.Num() != Transforms.Num())
+			{
+				return;
+			}
+
+			if (NumPieces != AssetArrays->BoundsBuffer.Num())
+			{
+				Init(Interface, SystemInstance);
+			}
+
+			for (int i = 0; i < NumPieces; ++i)
+			{
+				int32 TransformIndex = 3 * i;
+				AssetArrays->PrevWorldInverseTransformBuffer[TransformIndex] = AssetArrays->WorldInverseTransformBuffer[TransformIndex];
+				AssetArrays->PrevWorldInverseTransformBuffer[TransformIndex + 1] = AssetArrays->WorldInverseTransformBuffer[TransformIndex + 1];
+				AssetArrays->PrevWorldInverseTransformBuffer[TransformIndex + 2] = AssetArrays->WorldInverseTransformBuffer[TransformIndex + 2];
+
+				AssetArrays->PrevWorldTransformBuffer[TransformIndex] = AssetArrays->WorldTransformBuffer[TransformIndex];
+				AssetArrays->PrevWorldTransformBuffer[TransformIndex + 1] = AssetArrays->WorldTransformBuffer[TransformIndex + 1];
+				AssetArrays->PrevWorldTransformBuffer[TransformIndex + 2] = AssetArrays->WorldTransformBuffer[TransformIndex + 2];
+
+				FBox CurrBox = BoundingBoxes[i];
+
+				// #todo(dmp): save this somewhere in an array?
+				FVector3f LocalTranslation = (CurrBox.Max + CurrBox.Min) * .5;
+				FTransform LocalOffset(LocalTranslation);
+								
+				int32 CurrTransformIndex = TransformIndexArray[i];
+
+				FMatrix44f CurrTransform = LocalOffset.ToMatrixWithScale() * GlobalMatrices[CurrTransformIndex] * ActorTransform.ToMatrixWithScale();
+				CurrTransform.To3x4MatrixTranspose(&AssetArrays->WorldTransformBuffer[TransformIndex].X);
+
+				FMatrix44f CurrInverse = CurrTransform.Inverse();
+				CurrInverse.To3x4MatrixTranspose(&AssetArrays->WorldInverseTransformBuffer[TransformIndex].X);
+			}
+		}		
+	}
+}
+
+ETickingGroup FNDIGeometryCollectionData::ComputeTickingGroup()
+{
+	TickingGroup = NiagaraFirstTickGroup;
+	return TickingGroup;
+}
+
+
+//------------------------------------------------------------------------------------------------------------
+
+struct FNDIGeometryCollectionParametersCS : public FNiagaraDataInterfaceParametersCS
+{
+	DECLARE_TYPE_LAYOUT(FNDIGeometryCollectionParametersCS, NonVirtual);
+public:
+	void Bind(const FNiagaraDataInterfaceGPUParamInfo& ParameterInfo, const class FShaderParameterMap& ParameterMap)
+	{
+		FNDIGeometryCollectionParametersName ParamNames(*ParameterInfo.DataInterfaceHLSLSymbol);
+
+		NumPieces.Bind(ParameterMap, *ParamNames.NumPiecesName);
+		WorldTransformBuffer.Bind(ParameterMap, *ParamNames.WorldTransformBufferName);
+		PrevWorldTransformBuffer.Bind(ParameterMap, *ParamNames.PrevWorldTransformBufferName);
+		WorldInverseTransformBuffer.Bind(ParameterMap, *ParamNames.WorldInverseTransformBufferName);
+		PrevWorldInverseTransformBuffer.Bind(ParameterMap, *ParamNames.PrevWorldInverseTransformBufferName);
+		BoundsBuffer.Bind(ParameterMap, *ParamNames.BoundsBufferName);
+	}
+
+	void Set(FRHICommandList& RHICmdList, const FNiagaraDataInterfaceSetArgs& Context) const
+	{
+		check(IsInRenderingThread());
+
+		FRHIComputeShader* ComputeShaderRHI = RHICmdList.GetBoundComputeShader();
+
+		FNDIGeometryCollectionProxy* InterfaceProxy =
+			static_cast<FNDIGeometryCollectionProxy*>(Context.DataInterface);
+		FNDIGeometryCollectionData* ProxyData =
+			InterfaceProxy->SystemInstancesToProxyData.Find(Context.SystemInstanceID);
+
+		if (ProxyData != nullptr)
+		{
+
+			FNDIGeometryCollectionBuffer* AssetBuffer = ProxyData->AssetBuffer;
+
+			FRHITransitionInfo Transitions[] = {
+				FRHITransitionInfo(AssetBuffer->WorldTransformBuffer.UAV, ERHIAccess::Unknown, ERHIAccess::SRVCompute),
+				FRHITransitionInfo(AssetBuffer->PrevWorldTransformBuffer.UAV, ERHIAccess::Unknown, ERHIAccess::SRVCompute),
+				FRHITransitionInfo(AssetBuffer->WorldInverseTransformBuffer.UAV, ERHIAccess::Unknown, ERHIAccess::SRVCompute),
+				FRHITransitionInfo(AssetBuffer->PrevWorldInverseTransformBuffer.UAV, ERHIAccess::Unknown, ERHIAccess::SRVCompute),
+				FRHITransitionInfo(AssetBuffer->BoundsBuffer.UAV, ERHIAccess::Unknown, ERHIAccess::SRVCompute),
+			};
+			RHICmdList.Transition(MakeArrayView(Transitions, UE_ARRAY_COUNT(Transitions)));		
+			
+			SetShaderValue(RHICmdList, ComputeShaderRHI, NumPieces, ProxyData->AssetBuffer->NumPieces);
+			SetSRVParameter(RHICmdList, ComputeShaderRHI, WorldTransformBuffer, ProxyData->AssetBuffer->WorldTransformBuffer.SRV);
+			SetSRVParameter(RHICmdList, ComputeShaderRHI, PrevWorldTransformBuffer, ProxyData->AssetBuffer->PrevWorldTransformBuffer.SRV);
+			SetSRVParameter(RHICmdList, ComputeShaderRHI, WorldInverseTransformBuffer, ProxyData->AssetBuffer->WorldInverseTransformBuffer.SRV);
+			SetSRVParameter(RHICmdList, ComputeShaderRHI, PrevWorldInverseTransformBuffer, ProxyData->AssetBuffer->PrevWorldInverseTransformBuffer.SRV);
+			SetSRVParameter(RHICmdList, ComputeShaderRHI, BoundsBuffer, ProxyData->AssetBuffer->BoundsBuffer.SRV);
+		}
+		else
+		{
+			SetShaderValue(RHICmdList, ComputeShaderRHI, NumPieces, 0);
+			SetSRVParameter(RHICmdList, ComputeShaderRHI, WorldTransformBuffer, FNiagaraRenderer::GetDummyFloat4Buffer());
+			SetSRVParameter(RHICmdList, ComputeShaderRHI, PrevWorldTransformBuffer, FNiagaraRenderer::GetDummyFloat4Buffer());
+			SetSRVParameter(RHICmdList, ComputeShaderRHI, WorldInverseTransformBuffer, FNiagaraRenderer::GetDummyFloat4Buffer());
+			SetSRVParameter(RHICmdList, ComputeShaderRHI, PrevWorldInverseTransformBuffer, FNiagaraRenderer::GetDummyFloat4Buffer());
+			SetSRVParameter(RHICmdList, ComputeShaderRHI, BoundsBuffer, FNiagaraRenderer::GetDummyFloat4Buffer());
+		}
+	}
+
+	void Unset(FRHICommandList& RHICmdList, const FNiagaraDataInterfaceSetArgs& Context) const
+	{
+	}
+
+private:
+
+	LAYOUT_FIELD(FShaderParameter, NumPieces);
+
+	LAYOUT_FIELD(FShaderResourceParameter, WorldTransformBuffer);
+	LAYOUT_FIELD(FShaderResourceParameter, PrevWorldTransformBuffer);
+	LAYOUT_FIELD(FShaderResourceParameter, WorldInverseTransformBuffer);
+	LAYOUT_FIELD(FShaderResourceParameter, PrevWorldInverseTransformBuffer);
+	LAYOUT_FIELD(FShaderResourceParameter, BoundsBuffer);
+};
+
+IMPLEMENT_TYPE_LAYOUT(FNDIGeometryCollectionParametersCS);
+
+IMPLEMENT_NIAGARA_DI_PARAMETER(UNiagaraDataInterfaceGeometryCollection, FNDIGeometryCollectionParametersCS);
+
+
+//------------------------------------------------------------------------------------------------------------
+
+void FNDIGeometryCollectionProxy::ConsumePerInstanceDataFromGameThread(void* PerInstanceData, const FNiagaraSystemInstanceID& Instance)
+{
+	check(IsInRenderingThread());
+
+	FNDIGeometryCollectionData* SourceData = static_cast<FNDIGeometryCollectionData*>(PerInstanceData);	
+	FNDIGeometryCollectionData* TargetData = &(SystemInstancesToProxyData.FindOrAdd(Instance));
+
+	ensure(TargetData);
+	if (TargetData)
+	{
+		TargetData->AssetBuffer = SourceData->AssetBuffer;		
+		TargetData->AssetArrays = SourceData->AssetArrays;
+		TargetData->TickingGroup = SourceData->TickingGroup;
+	}
+	else
+	{
+		UE_LOG(LogGeometryCollection, Log, TEXT("ConsumePerInstanceDataFromGameThread() ... could not find %d"), Instance);
+	}
+}
+
+void FNDIGeometryCollectionProxy::InitializePerInstanceData(const FNiagaraSystemInstanceID& SystemInstance)
+{
+	check(IsInRenderingThread());
+
+	check(!SystemInstancesToProxyData.Contains(SystemInstance));
+	FNDIGeometryCollectionData* TargetData = &SystemInstancesToProxyData.Add(SystemInstance);
+}
+
+void FNDIGeometryCollectionProxy::DestroyPerInstanceData(const FNiagaraSystemInstanceID& SystemInstance)
+{
+	check(IsInRenderingThread());	
+
+	SystemInstancesToProxyData.Remove(SystemInstance);
+}
+
+void FNDIGeometryCollectionProxy::PreStage(FRHICommandList& RHICmdList, const FNiagaraDataInterfaceStageArgs& Context)
+{
+	check(SystemInstancesToProxyData.Contains(Context.SystemInstanceID));
+
+	FNDIGeometryCollectionData* ProxyData =
+		SystemInstancesToProxyData.Find(Context.SystemInstanceID);
+
+	if (ProxyData != nullptr && ProxyData->AssetBuffer)
+	{
+		if (Context.SimStageData->bFirstStage)
+		{
+			FNDIGeometryCollectionBuffer* AssetBuffer = ProxyData->AssetBuffer;
+
+			FRHITransitionInfo Transitions[] = {
+				FRHITransitionInfo(AssetBuffer->WorldTransformBuffer.UAV, ERHIAccess::Unknown, ERHIAccess::UAVCompute),
+				FRHITransitionInfo(AssetBuffer->PrevWorldTransformBuffer.UAV, ERHIAccess::Unknown, ERHIAccess::UAVCompute),
+				FRHITransitionInfo(AssetBuffer->WorldInverseTransformBuffer.UAV, ERHIAccess::Unknown, ERHIAccess::UAVCompute),
+				FRHITransitionInfo(AssetBuffer->PrevWorldInverseTransformBuffer.UAV, ERHIAccess::Unknown, ERHIAccess::UAVCompute),
+				FRHITransitionInfo(AssetBuffer->BoundsBuffer.UAV, ERHIAccess::Unknown, ERHIAccess::UAVCompute),
+			};
+			RHICmdList.Transition(MakeArrayView(Transitions, UE_ARRAY_COUNT(Transitions)));
+
+			// #todo(dmp): bounds buffer doesn't need to be updated each frame
+			UpdateInternalBuffer<FVector4f, EPixelFormat::PF_A32B32G32R32F>(ProxyData->AssetArrays->WorldTransformBuffer, ProxyData->AssetBuffer->WorldTransformBuffer);
+			UpdateInternalBuffer<FVector4f, EPixelFormat::PF_A32B32G32R32F>(ProxyData->AssetArrays->PrevWorldTransformBuffer, ProxyData->AssetBuffer->PrevWorldTransformBuffer);
+			UpdateInternalBuffer<FVector4f, EPixelFormat::PF_A32B32G32R32F>(ProxyData->AssetArrays->WorldInverseTransformBuffer, ProxyData->AssetBuffer->WorldInverseTransformBuffer);
+			UpdateInternalBuffer<FVector4f, EPixelFormat::PF_A32B32G32R32F>(ProxyData->AssetArrays->PrevWorldInverseTransformBuffer, ProxyData->AssetBuffer->PrevWorldInverseTransformBuffer);
+			UpdateInternalBuffer<FVector4f, EPixelFormat::PF_A32B32G32R32F>(ProxyData->AssetArrays->BoundsBuffer, ProxyData->AssetBuffer->BoundsBuffer);
+		}
+	}
+}
+
+void FNDIGeometryCollectionProxy::ResetData(FRHICommandList& RHICmdList, const FNiagaraDataInterfaceArgs& Context)
+{}
+
+//------------------------------------------------------------------------------------------------------------
+
+UNiagaraDataInterfaceGeometryCollection::UNiagaraDataInterfaceGeometryCollection(FObjectInitializer const& ObjectInitializer)
+	: Super(ObjectInitializer)	
+{
+	Proxy.Reset(new FNDIGeometryCollectionProxy());
+}
+
+bool UNiagaraDataInterfaceGeometryCollection::InitPerInstanceData(void* PerInstanceData, FNiagaraSystemInstance* SystemInstance)
+{
+	FNDIGeometryCollectionData* InstanceData = new (PerInstanceData) FNDIGeometryCollectionData();
+
+	check(InstanceData);
+	InstanceData->Init(this, SystemInstance);
+	
+	return true;
+}
+
+ETickingGroup UNiagaraDataInterfaceGeometryCollection::CalculateTickGroup(const void* PerInstanceData) const
+{
+	const FNDIGeometryCollectionData* InstanceData = static_cast<const FNDIGeometryCollectionData*>(PerInstanceData);
+
+	if (InstanceData)
+	{
+		return InstanceData->TickingGroup;
+	}
+	return NiagaraFirstTickGroup;
+}
+
+void UNiagaraDataInterfaceGeometryCollection::DestroyPerInstanceData(void* PerInstanceData, FNiagaraSystemInstance* SystemInstance)
+{
+	FNDIGeometryCollectionData* InstanceData = static_cast<FNDIGeometryCollectionData*>(PerInstanceData);
+
+	InstanceData->Release();	
+	InstanceData->~FNDIGeometryCollectionData();
+
+	FNDIGeometryCollectionProxy* ThisProxy = GetProxyAs<FNDIGeometryCollectionProxy>();
+	ENQUEUE_RENDER_COMMAND(FNiagaraDIDestroyInstanceData) (
+		[ThisProxy, InstanceID = SystemInstance->GetId()](FRHICommandListImmediate& CmdList)
+	{
+		FNDIGeometryCollectionData* ProxyData =
+			ThisProxy->SystemInstancesToProxyData.Find(InstanceID);
+
+		if (ProxyData != nullptr && ProxyData->AssetArrays)
+		{			
+			ThisProxy->SystemInstancesToProxyData.Remove(InstanceID);
+			delete ProxyData->AssetArrays;
+		}		
+	}
+	);
+}
+
+bool UNiagaraDataInterfaceGeometryCollection::PerInstanceTick(void* PerInstanceData, FNiagaraSystemInstance* SystemInstance, float InDeltaSeconds)
+{
+	FNDIGeometryCollectionData* InstanceData = static_cast<FNDIGeometryCollectionData*>(PerInstanceData);
+	if (InstanceData && InstanceData->AssetBuffer && SystemInstance)
+	{
+		InstanceData->Update(this, SystemInstance);
+	}
+	return false;
+}
+
+bool UNiagaraDataInterfaceGeometryCollection::CopyToInternal(UNiagaraDataInterface* Destination) const
+{
+	if (!Super::CopyToInternal(Destination))
+	{
+		return false;
+	}
+
+	UNiagaraDataInterfaceGeometryCollection* OtherTyped = CastChecked<UNiagaraDataInterfaceGeometryCollection>(Destination);		
+
+	OtherTyped->GeometryCollectionActor = GeometryCollectionActor;
+
+	return true;
+}
+
+bool UNiagaraDataInterfaceGeometryCollection::Equals(const UNiagaraDataInterface* Other) const
+{
+	if (!Super::Equals(Other))
+	{
+		return false;
+	}
+	const UNiagaraDataInterfaceGeometryCollection* OtherTyped = CastChecked<const UNiagaraDataInterfaceGeometryCollection>(Other);
+
+	return  OtherTyped->GeometryCollectionActor == GeometryCollectionActor;
+}
+
+void UNiagaraDataInterfaceGeometryCollection::PostInitProperties()
+{
+	Super::PostInitProperties();
+
+	if (HasAnyFlags(RF_ClassDefaultObject))
+	{
+		ENiagaraTypeRegistryFlags Flags = ENiagaraTypeRegistryFlags::AllowAnyVariable | ENiagaraTypeRegistryFlags::AllowParameter;
+		FNiagaraTypeRegistry::Register(FNiagaraTypeDefinition(GetClass()), Flags);
+	}
+}
+
+void UNiagaraDataInterfaceGeometryCollection::GetFunctions(TArray<FNiagaraFunctionSignature>& OutFunctions)
+{
+	{
+		FNiagaraFunctionSignature Sig;
+		Sig.Name = GetClosestPointNoNormalName;
+		Sig.bSupportsGPU = true;
+		Sig.bSupportsCPU = false;
+		Sig.bMemberFunction = true;
+		Sig.bRequiresContext = false;
+		Sig.Inputs.Add(FNiagaraVariable(FNiagaraTypeDefinition(GetClass()), TEXT("Collision DI")));
+		Sig.Inputs.Add(FNiagaraVariable(FNiagaraTypeDefinition::GetVec3Def(), TEXT("World Position")));
+		Sig.Inputs.Add(FNiagaraVariable(FNiagaraTypeDefinition::GetFloatDef(), TEXT("Delta Time")));
+		Sig.Inputs.Add(FNiagaraVariable(FNiagaraTypeDefinition::GetFloatDef(), TEXT("Time Fraction")));
+		Sig.Outputs.Add(FNiagaraVariable(FNiagaraTypeDefinition::GetFloatDef(), TEXT("Closest Distance")));
+		Sig.Outputs.Add(FNiagaraVariable(FNiagaraTypeDefinition::GetVec3Def(), TEXT("Closest Position")));		
+		Sig.Outputs.Add(FNiagaraVariable(FNiagaraTypeDefinition::GetVec3Def(), TEXT("Closest Velocity")));
+
+		OutFunctions.Add(Sig);
+	}
+}
+
+void UNiagaraDataInterfaceGeometryCollection::GetVMExternalFunction(const FVMExternalFunctionBindingInfo& BindingInfo, void* InstanceData, FVMExternalFunction& OutFunc)
+{
+}
+
+#if WITH_EDITORONLY_DATA
+bool UNiagaraDataInterfaceGeometryCollection::GetFunctionHLSL(const FNiagaraDataInterfaceGPUParamInfo& ParamInfo, const FNiagaraDataInterfaceGeneratedFunction& FunctionInfo, int FunctionInstanceIndex, FString& OutHLSL)
+{
+	FNDIGeometryCollectionParametersName ParamNames(ParamInfo.DataInterfaceHLSLSymbol);
+
+	TMap<FString, FStringFormatArg> ArgsSample = {
+
+		{TEXT("InstanceFunctionName"), FunctionInfo.InstanceName},
+		{TEXT("NumPiecesName"), ParamNames.NumPiecesName},
+		{TEXT("WorldTransformBufferName"), ParamNames.WorldTransformBufferName},
+		{TEXT("PrevWorldTransformBufferName"), ParamNames.PrevWorldTransformBufferName},
+		{TEXT("WorldInverseTransformBufferName"), ParamNames.WorldInverseTransformBufferName},
+		{TEXT("PrevWorldInverseTransformBufferName"), ParamNames.PrevWorldInverseTransformBufferName},
+		{TEXT("BoundsBufferName"), ParamNames.BoundsBufferName},
+		{TEXT("GeometryCollectionContextName"), TEXT("DIGEOMETRYCOLLECTION_MAKE_CONTEXT(") + ParamInfo.DataInterfaceHLSLSymbol + TEXT(")")},
+	};
+
+	if (FunctionInfo.DefinitionName == GetClosestPointNoNormalName)
+	{
+	static const TCHAR* FormatSample = TEXT(R"(
+		void {InstanceFunctionName}(in float3 WorldPosition, in float DeltaTime, in float TimeFraction, out float ClosestDistance, out float3 OutClosestPosition, 
+							out float3 OutClosestVelocity)
+		{
+			{GeometryCollectionContextName} DIGeometryCollection_GetClosestPointNoNormal(DIContext,WorldPosition,DeltaTime,TimeFraction, ClosestDistance,
+				OutClosestPosition,OutClosestVelocity);
+		}
+		)");
+	OutHLSL += FString::Format(FormatSample, ArgsSample);
+	return true;
+	}
+	OutHLSL += TEXT("\n");
+	return false;
+}
+
+void UNiagaraDataInterfaceGeometryCollection::GetCommonHLSL(FString& OutHLSL)
+{	
+	OutHLSL += TEXT("#include \"/Plugin/Experimental/ChaosNiagara/NiagaraDataInterfaceGeometryCollection.ush\"\n");			
+}
+
+void UNiagaraDataInterfaceGeometryCollection::GetParameterDefinitionHLSL(const FNiagaraDataInterfaceGPUParamInfo& ParamInfo, FString& OutHLSL)
+{
+	OutHLSL += TEXT("DIGEOMETRYCOLLECTION_DECLARE_CONSTANTS(") + ParamInfo.DataInterfaceHLSLSymbol + TEXT(")\n");
+}
+#endif
+
+#if WITH_EDITOR
+void UNiagaraDataInterfaceGeometryCollection::ValidateFunction(const FNiagaraFunctionSignature& Function, TArray<FText>& OutValidationErrors)
+{
+}
+#endif
+
+void UNiagaraDataInterfaceGeometryCollection::ProvidePerInstanceDataForRenderThread(void* DataForRenderThread, void* PerInstanceData, const FNiagaraSystemInstanceID& SystemInstance)
+{
+	FNDIGeometryCollectionData* GameThreadData = static_cast<FNDIGeometryCollectionData*>(PerInstanceData);
+	FNDIGeometryCollectionData* RenderThreadData = static_cast<FNDIGeometryCollectionData*>(DataForRenderThread);
+
+	if (GameThreadData != nullptr && RenderThreadData != nullptr)
+	{		
+		RenderThreadData->AssetBuffer = GameThreadData->AssetBuffer;				
+
+		RenderThreadData->AssetArrays = new FNDIGeometryCollectionArrays();
+		RenderThreadData->AssetArrays->CopyFrom(GameThreadData->AssetArrays);		
+		RenderThreadData->TickingGroup = GameThreadData->TickingGroup;
+	}
+	check(Proxy);
+}
+
+#undef LOCTEXT_NAMESPACE
\ No newline at end of file
diff --git a/Engine/Plugins/Experimental/ChaosNiagara/Source/ChaosNiagara/Public/NiagaraDataInterfaceGeometryCollection.h b/Engine/Plugins/Experimental/ChaosNiagara/Source/ChaosNiagara/Public/NiagaraDataInterfaceGeometryCollection.h
new file mode 100644
--- /dev/null
+++ b/Engine/Plugins/Experimental/ChaosNiagara/Source/ChaosNiagara/Public/NiagaraDataInterfaceGeometryCollection.h
@@ -0,0 +1,194 @@
+// Copyright Epic Games, Inc. All Rights Reserved.
+
+#pragma once
+
+#include "NiagaraDataInterface.h"
+#include "NiagaraCommon.h"
+#include "VectorVM.h"
+#include "GeometryCollection/GeometryCollectionActor.h"
+
+#include "NiagaraDataInterfaceGeometryCollection.generated.h"
+
+
+/** Arrays in which the cpu datas will be str */
+struct FNDIGeometryCollectionArrays
+{
+	TArray<FVector4f> WorldTransformBuffer;
+	TArray<FVector4f> PrevWorldTransformBuffer;
+	TArray<FVector4f> WorldInverseTransformBuffer;
+	TArray<FVector4f> PrevWorldInverseTransformBuffer;
+	TArray<FVector4f> BoundsBuffer;
+
+	FNDIGeometryCollectionArrays()
+	{
+		Resize(100);
+	}
+
+	FNDIGeometryCollectionArrays(uint32 Num)
+	{
+		Resize(Num);
+	}
+
+	void CopyFrom(const FNDIGeometryCollectionArrays* Other)
+	{
+		Resize(Other->NumPieces);
+		
+		WorldTransformBuffer = Other->WorldTransformBuffer;
+		PrevWorldTransformBuffer = Other->PrevWorldTransformBuffer;
+		WorldInverseTransformBuffer = Other->WorldInverseTransformBuffer;
+		PrevWorldInverseTransformBuffer = Other->PrevWorldInverseTransformBuffer;
+		BoundsBuffer = Other->BoundsBuffer;
+	}
+
+	void Resize(uint32 Num)
+	{		
+		NumPieces = Num;
+
+		WorldTransformBuffer.Init(FVector4f(0, 0, 0, 0), 3 * NumPieces);
+		PrevWorldTransformBuffer.Init(FVector4f(0, 0, 0, 0), 3 * NumPieces);
+		WorldInverseTransformBuffer.Init(FVector4f(0, 0, 0, 0), 3 * NumPieces);
+		PrevWorldInverseTransformBuffer.Init(FVector4f(0, 0, 0, 0), 3 * NumPieces);
+		BoundsBuffer.Init(FVector4f(0, 0, 0, 0), NumPieces);
+	}
+
+	uint32 NumPieces = 100;
+};
+
+/** Render buffers that will be used in hlsl functions */
+struct FNDIGeometryCollectionBuffer : public FRenderResource
+{
+	/** Init the buffer */
+	virtual void InitRHI() override;
+
+	/** Release the buffer */
+	virtual void ReleaseRHI() override;
+
+	/** Get the resource name */
+	virtual FString GetFriendlyName() const override { return TEXT("FNDIGeometryCollectionBuffer"); }
+
+	/** World transform buffer */
+	FRWBuffer WorldTransformBuffer;
+
+	/** Inverse transform buffer*/
+	FRWBuffer PrevWorldTransformBuffer;
+
+	/** World transform buffer */
+	FRWBuffer WorldInverseTransformBuffer;
+
+	/** Inverse transform buffer*/
+	FRWBuffer PrevWorldInverseTransformBuffer;
+
+	/** Element extent buffer */
+	FRWBuffer BoundsBuffer;
+
+	/** number of transforms */
+	uint32 NumPieces;
+
+	void SetNumPieces(uint32 Num)
+	{
+		NumPieces = Num;
+	}
+};
+
+/** Data stored per physics asset instance*/
+struct FNDIGeometryCollectionData
+{
+	/** Initialize the cpu datas */
+	void Init(class UNiagaraDataInterfaceGeometryCollection* Interface, FNiagaraSystemInstance* SystemInstance);
+
+	/** Update the gpu datas */
+	void Update(class UNiagaraDataInterfaceGeometryCollection* Interface, FNiagaraSystemInstance* SystemInstance);
+
+	/** Release the buffers */
+	void Release();
+
+	ETickingGroup ComputeTickingGroup();
+
+	/** The instance ticking group */
+	ETickingGroup TickingGroup;
+
+	/** Physics asset Gpu buffer */
+	FNDIGeometryCollectionBuffer* AssetBuffer = nullptr;
+
+	/** Physics asset Cpu arrays */
+	FNDIGeometryCollectionArrays *AssetArrays = nullptr;	
+};
+
+/** Data Interface for the Collisions */
+UCLASS(EditInlineNew, Category = "Collision", meta = (DisplayName = "Geometry Collection"))
+class UNiagaraDataInterfaceGeometryCollection : public UNiagaraDataInterface
+{
+	GENERATED_UCLASS_BODY()
+
+public:
+
+	DECLARE_NIAGARA_DI_PARAMETER();
+
+	UPROPERTY(EditAnywhere, Category = "Geometry Collection")
+	AGeometryCollectionActor *GeometryCollectionActor;
+
+
+	/** UObject Interface */
+	virtual void PostInitProperties() override;
+
+	/** UNiagaraDataInterface Interface */
+	virtual void GetFunctions(TArray<FNiagaraFunctionSignature>& OutFunctions) override;
+	virtual void GetVMExternalFunction(const FVMExternalFunctionBindingInfo& BindingInfo, void* InstanceData, FVMExternalFunction& OutFunc) override;
+	virtual bool CanExecuteOnTarget(ENiagaraSimTarget Target) const override { return Target == ENiagaraSimTarget::GPUComputeSim; }
+	virtual bool InitPerInstanceData(void* PerInstanceData, FNiagaraSystemInstance* SystemInstance) override;
+	virtual void DestroyPerInstanceData(void* PerInstanceData, FNiagaraSystemInstance* SystemInstance)override;
+	virtual bool PerInstanceTick(void* PerInstanceData, FNiagaraSystemInstance* SystemInstance, float DeltaSeconds) override;
+	virtual int32 PerInstanceDataSize()const override { return sizeof(FNDIGeometryCollectionData); }
+	virtual bool Equals(const UNiagaraDataInterface* Other) const override;
+	virtual bool HasPreSimulateTick() const override { return true; }
+	virtual bool HasTickGroupPrereqs() const override { return true; }
+	virtual ETickingGroup CalculateTickGroup(const void* PerInstanceData) const override;
+
+	virtual bool RequiresDistanceFieldData() const override { return true; }
+
+	/** GPU simulation  functionality */
+#if WITH_EDITORONLY_DATA
+	virtual void GetCommonHLSL(FString& OutHLSL) override;
+	virtual void GetParameterDefinitionHLSL(const FNiagaraDataInterfaceGPUParamInfo& ParamInfo, FString& OutHLSL) override;
+	virtual bool GetFunctionHLSL(const FNiagaraDataInterfaceGPUParamInfo& ParamInfo, const FNiagaraDataInterfaceGeneratedFunction& FunctionInfo, int FunctionInstanceIndex, FString& OutHLSL) override;
+
+	virtual void ValidateFunction(const FNiagaraFunctionSignature& Function, TArray<FText>& OutValidationErrors) override;
+#endif
+	virtual void ProvidePerInstanceDataForRenderThread(void* DataForRenderThread, void* PerInstanceData, const FNiagaraSystemInstanceID& SystemInstance) override;
+
+	static const FString NumPiecesName;
+	static const FString WorldTransformBufferName;
+	static const FString PrevWorldTransformBufferName;
+	static const FString WorldInverseTransformBufferName;
+	static const FString PrevWorldInverseTransformBufferName;
+	static const FString BoundsBufferName;
+
+protected:
+	/** Copy one niagara DI to this */
+	virtual bool CopyToInternal(UNiagaraDataInterface* Destination) const override;
+};
+
+/** Proxy to send data to gpu */
+struct FNDIGeometryCollectionProxy : public FNiagaraDataInterfaceProxy
+{
+	/** Get the size of the data that will be passed to render*/
+	virtual int32 PerInstanceDataPassedToRenderThreadSize() const override { return sizeof(FNDIGeometryCollectionData); }
+
+	/** Get the data that will be passed to render*/
+	virtual void ConsumePerInstanceDataFromGameThread(void* PerInstanceData, const FNiagaraSystemInstanceID& Instance) override;
+
+	/** Initialize the Proxy data buffer */
+	void InitializePerInstanceData(const FNiagaraSystemInstanceID& SystemInstance);
+
+	/** Destroy the proxy data if necessary */
+	void DestroyPerInstanceData(const FNiagaraSystemInstanceID& SystemInstance);
+
+	/** Launch all pre stage functions */
+	virtual void PreStage(FRHICommandList& RHICmdList, const FNiagaraDataInterfaceStageArgs& Context) override;
+
+	/** Reset the buffers  */
+	virtual void ResetData(FRHICommandList& RHICmdList, const FNiagaraDataInterfaceArgs& Context) override;
+
+	/** List of proxy data for each system instances*/
+	TMap<FNiagaraSystemInstanceID, FNDIGeometryCollectionData> SystemInstancesToProxyData;
+};
\ No newline at end of file
