Fixed a double-slash in NiagaraDataInterfaceGeometryCollection.cpp's header inclusion

diff --git a/Engine/Plugins/Experimental/ChaosNiagara/Source/ChaosNiagara/Private/NiagaraDataInterfaceGeometryCollection.cpp b/Engine/Plugins/Experimental/ChaosNiagara/Source/ChaosNiagara/Private/NiagaraDataInterfaceGeometryCollection.cpp
--- a/Engine/Plugins/Experimental/ChaosNiagara/Source/ChaosNiagara/Private/NiagaraDataInterfaceGeometryCollection.cpp
+++ b/Engine/Plugins/Experimental/ChaosNiagara/Source/ChaosNiagara/Private/NiagaraDataInterfaceGeometryCollection.cpp
@@ -17,7 +17,7 @@
 #include "DistanceFieldAtlas.h"
 #include "Renderer/Private/DistanceFieldLightingShared.h"
 #include "NiagaraGpuComputeDispatchInterface.h"
-#include "GeometryCollection//GeometryCollectionComponent.h"
+#include "GeometryCollection/GeometryCollectionComponent.h"
 
 #define LOCTEXT_NAMESPACE "NiagaraDataInterfaceGeometryCollection"
 DEFINE_LOG_CATEGORY_STATIC(LogGeometryCollection, Log, All);
