check against GC bounding box before looping over pieces.  Good speed bump

diff --git a/Engine/Plugins/Experimental/ChaosNiagara/Shaders/NiagaraDataInterfaceGeometryCollection.ush b/Engine/Plugins/Experimental/ChaosNiagara/Shaders/NiagaraDataInterfaceGeometryCollection.ush
--- a/Engine/Plugins/Experimental/ChaosNiagara/Shaders/NiagaraDataInterfaceGeometryCollection.ush
+++ b/Engine/Plugins/Experimental/ChaosNiagara/Shaders/NiagaraDataInterfaceGeometryCollection.ush
@@ -19,6 +19,8 @@ NiagaraDataInterfaceGeometryCollection.ush
 
 struct FDIGeometryCollectionContext
 {
+	float3 BoundsMin;
+	float3 BoundsMax;
 	uint NumPieces;	
 	Buffer<float4>		WorldTransformBuffer;
 	Buffer<float4>		PrevWorldTransformBuffer;
@@ -28,6 +30,8 @@ struct FDIGeometryCollectionContext
 }; 
 
 #define DIGEOMETRYCOLLECTION_DECLARE_CONSTANTS(NAME)\
+float3 BoundsMin_##NAME;\
+float3 BoundsMax_##NAME;\
 uint NumPieces_##NAME;\
 Buffer<float4>		WorldTransformBuffer_##NAME;\
 Buffer<float4>		PrevWorldTransformBuffer_##NAME;\
@@ -38,6 +42,8 @@ Buffer<float4>		BoundsBuffer_##NAME;\
 
 #define DIGEOMETRYCOLLECTION_MAKE_CONTEXT(NAME)\
 FDIGeometryCollectionContext DIContext; \
+DIContext.BoundsMin = BoundsMin_##NAME;\
+DIContext.BoundsMax = BoundsMax_##NAME;\
 DIContext.NumPieces = NumPieces_##NAME;\
 DIContext.WorldTransformBuffer =  WorldTransformBuffer_##NAME;\
 DIContext.PrevWorldTransformBuffer = PrevWorldTransformBuffer_##NAME;\
@@ -195,6 +201,19 @@ void DIGeometryCollection_GetClosestPointNoNormal(in FDIGeometryCollectionContex
 	OutClosestPosition = float3(0,0,0);
 	OutClosestVelocity = float3(0,0,0);
 	
+	
+	[branch]
+	if (WorldPosition.x < DIContext.BoundsMin.x || 
+		WorldPosition.y < DIContext.BoundsMin.y || 
+		WorldPosition.z < DIContext.BoundsMin.z || 
+		WorldPosition.x > DIContext.BoundsMax.x || 
+		WorldPosition.y > DIContext.BoundsMax.y || 
+		WorldPosition.z > DIContext.BoundsMax.z)
+	{
+		return;
+	}
+	
+
 	int ClosestTransform = -1;
 	for (int CurrTransform = 0; CurrTransform < DIContext.NumPieces; ++CurrTransform)
 	{
diff --git a/Engine/Plugins/Experimental/ChaosNiagara/Source/ChaosNiagara/Private/NiagaraDataInterfaceGeometryCollection.cpp b/Engine/Plugins/Experimental/ChaosNiagara/Source/ChaosNiagara/Private/NiagaraDataInterfaceGeometryCollection.cpp
--- a/Engine/Plugins/Experimental/ChaosNiagara/Source/ChaosNiagara/Private/NiagaraDataInterfaceGeometryCollection.cpp
+++ b/Engine/Plugins/Experimental/ChaosNiagara/Source/ChaosNiagara/Private/NiagaraDataInterfaceGeometryCollection.cpp
@@ -29,6 +29,8 @@ static const FName GetClosestPointNoNormalName(TEXT("GetClosestPointNoNormal"));
 
 //------------------------------------------------------------------------------------------------------------
 
+const FString UNiagaraDataInterfaceGeometryCollection::BoundsMinName(TEXT("BoundsMin_"));
+const FString UNiagaraDataInterfaceGeometryCollection::BoundsMaxName(TEXT("BoundsMax_"));
 const FString UNiagaraDataInterfaceGeometryCollection::NumPiecesName(TEXT("NumPieces_"));
 const FString UNiagaraDataInterfaceGeometryCollection::WorldTransformBufferName(TEXT("WorldTransformBuffer_"));
 const FString UNiagaraDataInterfaceGeometryCollection::PrevWorldTransformBufferName(TEXT("PrevWorldTransformBuffer_"));
@@ -43,6 +45,8 @@ struct FNDIGeometryCollectionParametersName
 {
 	FNDIGeometryCollectionParametersName(const FString& Suffix)
 	{
+		BoundsMinName = UNiagaraDataInterfaceGeometryCollection::BoundsMinName + Suffix;
+		BoundsMaxName = UNiagaraDataInterfaceGeometryCollection::BoundsMaxName + Suffix;
 		NumPiecesName = UNiagaraDataInterfaceGeometryCollection::NumPiecesName + Suffix;
 		WorldTransformBufferName = UNiagaraDataInterfaceGeometryCollection::WorldTransformBufferName + Suffix;
 		PrevWorldTransformBufferName = UNiagaraDataInterfaceGeometryCollection::PrevWorldTransformBufferName + Suffix;
@@ -51,6 +55,9 @@ struct FNDIGeometryCollectionParametersName
 		BoundsBufferName = UNiagaraDataInterfaceGeometryCollection::BoundsBufferName + Suffix;
 	}
 
+	FString BoundsMinName;
+	FString BoundsMaxName;
+
 	FString NumPiecesName;
 	FString WorldTransformBufferName;
 	FString PrevWorldTransformBufferName;
@@ -147,7 +154,14 @@ void FNDIGeometryCollectionData::Init(UNiagaraDataInterfaceGeometryCollection* I
 			AssetBuffer = new FNDIGeometryCollectionBuffer();
 			AssetBuffer->SetNumPieces(NumPieces);
 			BeginInitResource(AssetBuffer);
-			
+
+			FVector Origin;
+			FVector Extents;
+			Interface->GeometryCollectionActor->GetActorBounds(false, Origin, Extents, true);
+
+			BoundsOrigin = Origin;
+			BoundsExtent = Extents;
+
 			for (int i = 0; i < NumPieces; ++i)
 			{
 				FBox CurrBox = BoundingBoxes[i];
@@ -198,6 +212,13 @@ void FNDIGeometryCollectionData::Update(UNiagaraDataInterfaceGeometryCollection*
 			{
 				Init(Interface, SystemInstance);
 			}
+			
+			FVector Origin;
+			FVector Extents;
+			Interface->GeometryCollectionActor->GetActorBounds(false, Origin, Extents, true);
+
+			BoundsOrigin = Origin;
+			BoundsExtent = Extents;
 
 			for (int i = 0; i < NumPieces; ++i)
 			{
@@ -245,6 +266,9 @@ public:
 	{
 		FNDIGeometryCollectionParametersName ParamNames(*ParameterInfo.DataInterfaceHLSLSymbol);
 
+		BoundsMin.Bind(ParameterMap, *ParamNames.BoundsMinName);
+		BoundsMax.Bind(ParameterMap, *ParamNames.BoundsMaxName);
+
 		NumPieces.Bind(ParameterMap, *ParamNames.NumPiecesName);
 		WorldTransformBuffer.Bind(ParameterMap, *ParamNames.WorldTransformBufferName);
 		PrevWorldTransformBuffer.Bind(ParameterMap, *ParamNames.PrevWorldTransformBufferName);
@@ -278,6 +302,9 @@ public:
 			};
 			RHICmdList.Transition(MakeArrayView(Transitions, UE_ARRAY_COUNT(Transitions)));		
 			
+			SetShaderValue(RHICmdList, ComputeShaderRHI, BoundsMin, ProxyData->BoundsOrigin - ProxyData->BoundsExtent);
+			SetShaderValue(RHICmdList, ComputeShaderRHI, BoundsMax, ProxyData->BoundsOrigin + ProxyData->BoundsExtent);
+
 			SetShaderValue(RHICmdList, ComputeShaderRHI, NumPieces, ProxyData->AssetBuffer->NumPieces);
 			SetSRVParameter(RHICmdList, ComputeShaderRHI, WorldTransformBuffer, ProxyData->AssetBuffer->WorldTransformBuffer.SRV);
 			SetSRVParameter(RHICmdList, ComputeShaderRHI, PrevWorldTransformBuffer, ProxyData->AssetBuffer->PrevWorldTransformBuffer.SRV);
@@ -287,6 +314,9 @@ public:
 		}
 		else
 		{
+			SetShaderValue(RHICmdList, ComputeShaderRHI, BoundsMin, FVector3f(0, 0, 0));
+			SetShaderValue(RHICmdList, ComputeShaderRHI, BoundsMax, FVector3f(0, 0, 0));
+
 			SetShaderValue(RHICmdList, ComputeShaderRHI, NumPieces, 0);
 			SetSRVParameter(RHICmdList, ComputeShaderRHI, WorldTransformBuffer, FNiagaraRenderer::GetDummyFloat4Buffer());
 			SetSRVParameter(RHICmdList, ComputeShaderRHI, PrevWorldTransformBuffer, FNiagaraRenderer::GetDummyFloat4Buffer());
@@ -302,6 +332,9 @@ public:
 
 private:
 
+	LAYOUT_FIELD(FShaderParameter, BoundsMin);
+	LAYOUT_FIELD(FShaderParameter, BoundsMax);
+
 	LAYOUT_FIELD(FShaderParameter, NumPieces);
 
 	LAYOUT_FIELD(FShaderResourceParameter, WorldTransformBuffer);
@@ -331,6 +364,8 @@ void FNDIGeometryCollectionProxy::ConsumePerInstanceDataFromGameThread(void* Per
 		TargetData->AssetBuffer = SourceData->AssetBuffer;		
 		TargetData->AssetArrays = SourceData->AssetArrays;
 		TargetData->TickingGroup = SourceData->TickingGroup;
+		TargetData->BoundsOrigin = SourceData->BoundsOrigin;
+		TargetData->BoundsExtent = SourceData->BoundsExtent;
 	}
 	else
 	{
@@ -574,6 +609,8 @@ void UNiagaraDataInterfaceGeometryCollection::ProvidePerInstanceDataForRenderThr
 		RenderThreadData->AssetArrays = new FNDIGeometryCollectionArrays();
 		RenderThreadData->AssetArrays->CopyFrom(GameThreadData->AssetArrays);		
 		RenderThreadData->TickingGroup = GameThreadData->TickingGroup;
+		RenderThreadData->BoundsOrigin = GameThreadData->BoundsOrigin;
+		RenderThreadData->BoundsExtent = GameThreadData->BoundsExtent;
 	}
 	check(Proxy);
 }
diff --git a/Engine/Plugins/Experimental/ChaosNiagara/Source/ChaosNiagara/Public/NiagaraDataInterfaceGeometryCollection.h b/Engine/Plugins/Experimental/ChaosNiagara/Source/ChaosNiagara/Public/NiagaraDataInterfaceGeometryCollection.h
--- a/Engine/Plugins/Experimental/ChaosNiagara/Source/ChaosNiagara/Public/NiagaraDataInterfaceGeometryCollection.h
+++ b/Engine/Plugins/Experimental/ChaosNiagara/Source/ChaosNiagara/Public/NiagaraDataInterfaceGeometryCollection.h
@@ -107,6 +107,11 @@ struct FNDIGeometryCollectionData
 	/** The instance ticking group */
 	ETickingGroup TickingGroup;
 
+	/** Geometry Collection Bounds */
+	FVector3f BoundsOrigin;
+
+	FVector3f BoundsExtent;
+
 	/** Physics asset Gpu buffer */
 	FNDIGeometryCollectionBuffer* AssetBuffer = nullptr;
 
@@ -155,6 +160,8 @@ public:
 #endif
 	virtual void ProvidePerInstanceDataForRenderThread(void* DataForRenderThread, void* PerInstanceData, const FNiagaraSystemInstanceID& SystemInstance) override;
 
+	static const FString BoundsMinName;
+	static const FString BoundsMaxName;
 	static const FString NumPiecesName;
 	static const FString WorldTransformBufferName;
 	static const FString PrevWorldTransformBufferName;
