Chaos and Niagara: fix crashes with physics field

diff --git a/Engine/Plugins/Experimental/ChaosNiagara/Source/ChaosNiagara/Classes/NiagaraDataInterfacePhysicsField.h b/Engine/Plugins/Experimental/ChaosNiagara/Source/ChaosNiagara/Classes/NiagaraDataInterfacePhysicsField.h
--- a/Engine/Plugins/Experimental/ChaosNiagara/Source/ChaosNiagara/Classes/NiagaraDataInterfacePhysicsField.h
+++ b/Engine/Plugins/Experimental/ChaosNiagara/Source/ChaosNiagara/Classes/NiagaraDataInterfacePhysicsField.h
@@ -9,7 +9,14 @@
 
 #include "NiagaraDataInterfacePhysicsField.generated.h"
 
-/** Data stored per physics asset instance*/
+/** Data stored per physics asset instance on the render thread */
+struct FNDIFieldRenderData
+{
+	/** Field render resource for gpu */
+	class FPhysicsFieldResource* FieldResource;
+};
+
+/** Data stored per physics asset instance on the game thread */
 struct FNDIPhysicsFieldData
 {
 	/** Initialize the resource */
@@ -99,7 +106,7 @@ protected:
 struct FNDIPhysicsFieldProxy : public FNiagaraDataInterfaceProxy
 {
 	/** Get the size of the data that will be passed to render*/
-	virtual int32 PerInstanceDataPassedToRenderThreadSize() const override { return sizeof(FNDIPhysicsFieldData); }
+	virtual int32 PerInstanceDataPassedToRenderThreadSize() const override { return sizeof(FNDIFieldRenderData); }
 
 	/** Get the data that will be passed to render*/
 	virtual void ConsumePerInstanceDataFromGameThread(void* PerInstanceData, const FNiagaraSystemInstanceID& Instance) override;
@@ -111,6 +118,6 @@ struct FNDIPhysicsFieldProxy : public FNiagaraDataInterfaceProxy
 	void DestroyPerInstanceData(const FNiagaraSystemInstanceID& SystemInstance);
 
 	/** List of proxy data for each system instances*/
-	TMap<FNiagaraSystemInstanceID, FNDIPhysicsFieldData> SystemInstancesToProxyData;
+	TMap<FNiagaraSystemInstanceID, FNDIFieldRenderData> SystemInstancesToProxyData;
 };
 
diff --git a/Engine/Plugins/Experimental/ChaosNiagara/Source/ChaosNiagara/Private/NiagaraDataInterfacePhysicsField.cpp b/Engine/Plugins/Experimental/ChaosNiagara/Source/ChaosNiagara/Private/NiagaraDataInterfacePhysicsField.cpp
--- a/Engine/Plugins/Experimental/ChaosNiagara/Source/ChaosNiagara/Private/NiagaraDataInterfacePhysicsField.cpp
+++ b/Engine/Plugins/Experimental/ChaosNiagara/Source/ChaosNiagara/Private/NiagaraDataInterfacePhysicsField.cpp
@@ -69,7 +69,7 @@ void FNDIPhysicsFieldData::Release()
 
 void FNDIPhysicsFieldData::Init(FNiagaraSystemInstance* SystemInstance)
 {
-	FieldResource = nullptr;
+	Release();
 	if (SystemInstance != nullptr)
 	{
 		UWorld* World = SystemInstance->GetWorld();
@@ -132,7 +132,7 @@ public:
 
 		FNDIPhysicsFieldProxy* InterfaceProxy =
 			static_cast<FNDIPhysicsFieldProxy*>(Context.DataInterface);
-		FNDIPhysicsFieldData* ProxyData =
+		FNDIFieldRenderData* ProxyData =
 			InterfaceProxy->SystemInstancesToProxyData.Find(Context.SystemInstanceID);
 
 		TStaticArray<FIntVector4, MAX_PHYSICS_FIELD_TARGETS,16> LocalTargets;
@@ -195,22 +195,22 @@ IMPLEMENT_NIAGARA_DI_PARAMETER(UNiagaraDataInterfacePhysicsField, FNDIPhysicsFie
 
 void FNDIPhysicsFieldProxy::ConsumePerInstanceDataFromGameThread(void* PerInstanceData, const FNiagaraSystemInstanceID& Instance)
 {
-	FNDIPhysicsFieldData* SourceData = static_cast<FNDIPhysicsFieldData*>(PerInstanceData);
-	FNDIPhysicsFieldData* TargetData = &(SystemInstancesToProxyData.FindOrAdd(Instance));
+	FNDIFieldRenderData* SourceData = static_cast<FNDIFieldRenderData*>(PerInstanceData);
+	FNDIFieldRenderData* TargetData = &(SystemInstancesToProxyData.FindOrAdd(Instance));
 
 	ensure(TargetData);
 	if (TargetData && SourceData && SourceData->FieldResource)
 	{
 		TargetData->FieldResource = SourceData->FieldResource;
 	}
-	SourceData->~FNDIPhysicsFieldData();
+	SourceData->~FNDIFieldRenderData();
 }
 
 void FNDIPhysicsFieldProxy::InitializePerInstanceData(const FNiagaraSystemInstanceID& SystemInstance)
 {
 	check(IsInRenderingThread());
 
-	FNDIPhysicsFieldData* TargetData = SystemInstancesToProxyData.Find(SystemInstance);
+	FNDIFieldRenderData* TargetData = SystemInstancesToProxyData.Find(SystemInstance);
 	TargetData = &SystemInstancesToProxyData.Add(SystemInstance);
 }
 
@@ -775,7 +775,7 @@ void UNiagaraDataInterfacePhysicsField::GetParameterDefinitionHLSL(const FNiagar
 void UNiagaraDataInterfacePhysicsField::ProvidePerInstanceDataForRenderThread(void* DataForRenderThread, void* PerInstanceData, const FNiagaraSystemInstanceID& SystemInstance)
 {
 	FNDIPhysicsFieldData* GameThreadData = static_cast<FNDIPhysicsFieldData*>(PerInstanceData);
-	FNDIPhysicsFieldData* RenderThreadData = static_cast<FNDIPhysicsFieldData*>(DataForRenderThread);
+	FNDIFieldRenderData* RenderThreadData = static_cast<FNDIFieldRenderData*>(DataForRenderThread);
 
 	if (GameThreadData != nullptr && RenderThreadData != nullptr)
 	{
