NiagaradDataInterface - Material Updates
Fixed/Added PhysicalMaterial support for Chaos Destruction Collision/Breaking.
Added example level showing how to setup a single Nigara sysytem emitting 3 different type of debris for GCs by using PhysicalMaterials.

diff --git a/Engine/Plugins/Experimental/ChaosNiagara/Source/ChaosNiagara/Classes/NiagaraDataInterfaceChaosDestruction.h b/Engine/Plugins/Experimental/ChaosNiagara/Source/ChaosNiagara/Classes/NiagaraDataInterfaceChaosDestruction.h
--- a/Engine/Plugins/Experimental/ChaosNiagara/Source/ChaosNiagara/Classes/NiagaraDataInterfaceChaosDestruction.h
+++ b/Engine/Plugins/Experimental/ChaosNiagara/Source/ChaosNiagara/Classes/NiagaraDataInterfaceChaosDestruction.h
@@ -478,7 +478,7 @@ public:
 	bool bApplyMaterialsFilter;
 
 	/* TODO: Explanatory comment */
-	UPROPERTY(EditAnywhere, Category = "Breaking Data Settings", meta = (DisplayName = "Breaking Filtered Materials", EditCondition = bApplyMaterialsFilter))
+	UPROPERTY(EditAnywhere, Category = "Breaking Data Settings", meta = (DisplayName = "Collision/Breaking Filtered Materials", EditCondition = bApplyMaterialsFilter))
 	TSet<TObjectPtr<UPhysicalMaterial>> ChaosBreakingMaterialSet;
 
 	/* TODO: Explanatory comment */
diff --git a/Engine/Plugins/Experimental/ChaosNiagara/Source/ChaosNiagara/Private/NiagaraDataInterfaceChaosDestruction.cpp b/Engine/Plugins/Experimental/ChaosNiagara/Source/ChaosNiagara/Private/NiagaraDataInterfaceChaosDestruction.cpp
--- a/Engine/Plugins/Experimental/ChaosNiagara/Source/ChaosNiagara/Private/NiagaraDataInterfaceChaosDestruction.cpp
+++ b/Engine/Plugins/Experimental/ChaosNiagara/Source/ChaosNiagara/Private/NiagaraDataInterfaceChaosDestruction.cpp
@@ -750,10 +750,80 @@ void UNiagaraDataInterfaceChaosDestruction::HandleCollisionEvents(const Chaos::F
 			AllCollisionsArray[Idx].SurfaceType);
 		*/
 
+		UPhysicalMaterial* PhysicalMaterial = nullptr;
+		int32 MaterialID = 0;
+
+		UGeometryCollectionComponent* GeometryCollectionComponent = nullptr;
+		UMaterialInterface* Material = nullptr;
+#if INCLUDE_CHAOS
+		for (auto& Solver : Solvers)
+		{
+			if (Solver.PhysScene)
+			{
+				GeometryCollectionComponent = Solver.PhysScene->GetOwningComponent<UGeometryCollectionComponent>(DataIn.Proxy1);
+				if (GeometryCollectionComponent)
+					break;
+			}
+		}
+
+		if (GeometryCollectionComponent)
+		{
+			Material = GeometryCollectionComponent->GetMaterial(MaterialID);
+		}
+#endif
+		if (Material)
+		{
+			PhysicalMaterial = Material->GetPhysicalMaterial();
+			ensure(PhysicalMaterial);
+			if (PhysicalMaterial)
+			{
+				CopyData.SurfaceType1 = PhysicalMaterial->SurfaceType;
+			}
+		}
+
+		if (PhysicalMaterial)
+		{
+			CopyData.PhysicalMaterialName1 = PhysicalMaterial->GetFName();
+		}
+		else
+		{
+			CopyData.PhysicalMaterialName1 = FName();
+		}
+
+#if INCLUDE_CHAOS
+		for (auto& Solver : Solvers)
+		{
+			GeometryCollectionComponent = Solver.PhysScene->GetOwningComponent<UGeometryCollectionComponent>(DataIn.Proxy2);
+			if (GeometryCollectionComponent)
+				break;
+		}
+		if (GeometryCollectionComponent)
+		{
+			Material = GeometryCollectionComponent->GetMaterial(MaterialID);
+		}
+#endif
+		if (Material)
+		{
+			PhysicalMaterial = Material->GetPhysicalMaterial();
+			ensure(PhysicalMaterial);
+			if (PhysicalMaterial)
+			{
+				CopyData.SurfaceType2 = PhysicalMaterial->SurfaceType;
+			}
+		}
+
+		if (PhysicalMaterial)
+		{
+			CopyData.PhysicalMaterialName2 = PhysicalMaterial->GetFName();
+		}
+		else
+		{
+			CopyData.PhysicalMaterialName2 = FName();
+		}
+		
 		CopyData.BoundingboxVolume = 1000000.f;
 		CopyData.BoundingboxExtentMin = 100.f;
 		CopyData.BoundingboxExtentMax = 100.f;
-		CopyData.SurfaceType = 0;
 
 		Idx++;
 	}
@@ -765,6 +835,7 @@ void UNiagaraDataInterfaceChaosDestruction::FilterAllCollisions(TArray<Chaos::FC
 	QUICK_SCOPE_CYCLE_COUNTER(STAT_FilterAllCollisions);
 
 	if (/*ParticleToProcess != nullptr ||*/
+		bApplyMaterialsFilter ||
 		ImpulseToSpawnMinMax.X > 0.f ||
 		ImpulseToSpawnMinMax.Y > 0.f ||
 		SpeedToSpawnMinMax.X > 0.f ||
@@ -794,12 +865,36 @@ void UNiagaraDataInterfaceChaosDestruction::FilterAllCollisions(TArray<Chaos::FC
 		float MinSpeedToSpawnSquared = SpeedToSpawnMinMax.X * SpeedToSpawnMinMax.X;
 		float MaxSpeedToSpawnSquared = SpeedToSpawnMinMax.Y * SpeedToSpawnMinMax.Y;
 
+		auto IsMaterialInFilter = [&](const FName& InMaterialName)
+		{
+			if (!InMaterialName.IsValid())
+			{
+				return false;
+			}
+
+			for (const UPhysicalMaterial* Material : ChaosBreakingMaterialSet)
+			{
+				if (!Material)
+				{
+					continue;
+				}
+
+				if (Material->GetFName() == InMaterialName)
+				{
+					return true;
+				}
+			}
+
+			return false;
+		};
+
 		for (int32 IdxCollision = 0; IdxCollision < AllCollisionsArray.Num(); ++IdxCollision)
 		{
 			float CollisionAccumulatedImpulseSquared = AllCollisionsArray[IdxCollision].AccumulatedImpulse.SizeSquared();
 			float CollisionSpeedSquared = AllCollisionsArray[IdxCollision].Velocity1.SizeSquared();
 
 			if (/*(ParticleToProcess != nullptr && AllCollisionsArrayInOut[IdxCollision].Particle != ParticleToProcess) ||*/
+				!(bApplyMaterialsFilter && (IsMaterialInFilter(AllCollisionsArray[IdxCollision].PhysicalMaterialName1) || IsMaterialInFilter(AllCollisionsArray[IdxCollision].PhysicalMaterialName2))) ||
 				(ImpulseToSpawnMinMax.X > 0.f && ImpulseToSpawnMinMax.Y < 0.f && CollisionAccumulatedImpulseSquared < MinImpulseToSpawnSquared) ||
 				(ImpulseToSpawnMinMax.X < 0.f && ImpulseToSpawnMinMax.Y > 0.f && CollisionAccumulatedImpulseSquared > MaxImpulseToSpawnSquared) ||
 				(ImpulseToSpawnMinMax.X > 0.f && ImpulseToSpawnMinMax.Y > 0.f && (CollisionAccumulatedImpulseSquared < MinImpulseToSpawnSquared || CollisionAccumulatedImpulseSquared > MaxImpulseToSpawnSquared)) ||
@@ -821,7 +916,7 @@ void UNiagaraDataInterfaceChaosDestruction::FilterAllCollisions(TArray<Chaos::FC
 				(SolverTimeToSpawnMinMax.X > 0.f && SolverTimeToSpawnMinMax.Y < 0.f && SolverTime < SolverTimeToSpawnMinMax.X) ||
 				(SolverTimeToSpawnMinMax.X < 0.f && SolverTimeToSpawnMinMax.Y > 0.f && SolverTime > SolverTimeToSpawnMinMax.Y) ||
 				(SolverTimeToSpawnMinMax.X > 0.f && SolverTimeToSpawnMinMax.Y > 0.f && (SolverTime < SolverTimeToSpawnMinMax.X || SolverTime > SolverTimeToSpawnMinMax.Y)) ||
-				(SurfaceTypeToSpawn != -1 && AllCollisionsArray[IdxCollision].SurfaceType != SurfaceTypeToSpawn) ||
+				(SurfaceTypeToSpawn != -1 && AllCollisionsArray[IdxCollision].SurfaceType1 != SurfaceTypeToSpawn && AllCollisionsArray[IdxCollision].SurfaceType2 != SurfaceTypeToSpawn) ||
 				(LocationFilteringMode == ELocationFilteringModeEnum::ChaosNiagara_LocationFilteringMode_Inclusive && LocationXToSpawn == ELocationXToSpawnEnum::ChaosNiagara_LocationXToSpawn_Min && AllCollisionsArray[IdxCollision].Location.X <= LocationXToSpawnMinMax.X) ||
 				(LocationFilteringMode == ELocationFilteringModeEnum::ChaosNiagara_LocationFilteringMode_Inclusive && LocationXToSpawn == ELocationXToSpawnEnum::ChaosNiagara_LocationXToSpawn_Max && AllCollisionsArray[IdxCollision].Location.X >= LocationXToSpawnMinMax.Y) ||
 				(LocationFilteringMode == ELocationFilteringModeEnum::ChaosNiagara_LocationFilteringMode_Inclusive && LocationXToSpawn == ELocationXToSpawnEnum::ChaosNiagara_LocationXToSpawn_MinMax && (AllCollisionsArray[IdxCollision].Location.X <= LocationXToSpawnMinMax.X || AllCollisionsArray[IdxCollision].Location.X >= LocationXToSpawnMinMax.Y)) ||
@@ -1203,7 +1298,7 @@ bool UNiagaraDataInterfaceChaosDestruction::CollisionCallback(FNDIChaosDestructi
 						// jf: optimization: presize these arrays?
 						for (int32 Idx = 0; Idx < NumParticlesSpawned; ++Idx)
 						{
-							InstData->SurfaceTypeArray.Add(CollisionsToSpawnArray[IdxCollision].SurfaceType);
+							InstData->SurfaceTypeArray.Add(CollisionsToSpawnArray[IdxCollision].SurfaceType1);
 							InstData->ColorArray.Add(Color);
 							InstData->FrictionArray.Add(Friction);
 							InstData->RestitutionArray.Add(Restitution);
@@ -1265,9 +1360,12 @@ void UNiagaraDataInterfaceChaosDestruction::HandleBreakingEvents(const Chaos::FB
 #if INCLUDE_CHAOS
 			for (auto& Solver : Solvers)
 			{
-				GeometryCollectionComponent = Solver.PhysScene->GetOwningComponent<UGeometryCollectionComponent>(DataIn.Proxy);
-				if (GeometryCollectionComponent)
-					break;
+				if (Solver.PhysScene)
+				{
+					GeometryCollectionComponent = Solver.PhysScene->GetOwningComponent<UGeometryCollectionComponent>(DataIn.Proxy);
+					if (GeometryCollectionComponent)
+						break;
+				}
 			}
 			if (GeometryCollectionComponent)
 			{
@@ -1284,9 +1382,9 @@ void UNiagaraDataInterfaceChaosDestruction::HandleBreakingEvents(const Chaos::FB
 				}
 			}
 
-			if (Material)
+			if (PhysicalMaterial)
 			{
-				CopyData.PhysicalMaterialName = Material->GetFName();
+				CopyData.PhysicalMaterialName = PhysicalMaterial->GetFName();
 			}
 			else
 			{
diff --git a/Engine/Source/Runtime/Experimental/Chaos/Public/Chaos/ExternalCollisionData.h b/Engine/Source/Runtime/Experimental/Chaos/Public/Chaos/ExternalCollisionData.h
--- a/Engine/Source/Runtime/Experimental/Chaos/Public/Chaos/ExternalCollisionData.h
+++ b/Engine/Source/Runtime/Experimental/Chaos/Public/Chaos/ExternalCollisionData.h
@@ -95,13 +95,14 @@ namespace Chaos
 			, BoundingboxVolume((FReal)-1.0)
 			, BoundingboxExtentMin((FReal)-1.0)
 			, BoundingboxExtentMax((FReal)-1.0)
-			, SurfaceType(-1)
+			, SurfaceType1(-1)
+			, SurfaceType2(-1)
 		{}
 
 		FCollidingDataExt(
 		    FVec3 InLocation, FVec3 InAccumulatedImpulse, FVec3 InNormal, FVec3 InVelocity1, FVec3 InVelocity2
 			, FVec3 InAngularVelocity1, FVec3 InAngularVelocity2, FReal InMass1, FReal InMass2 //, FGeometryParticleHandle* InParticle, FGeometryParticleHandle* InLevelset
-			, FReal InBoundingboxVolume, FReal InBoundingboxExtentMin, FReal InBoundingboxExtentMax, int32 InSurfaceType)
+			, FReal InBoundingboxVolume, FReal InBoundingboxExtentMin, FReal InBoundingboxExtentMax, int32 InSurfaceType1, int32 InSurfaceType2)
 			: Location(InLocation)
 			, AccumulatedImpulse(InAccumulatedImpulse)
 			, Normal(InNormal)
@@ -114,7 +115,8 @@ namespace Chaos
 			, BoundingboxVolume(InBoundingboxVolume)
 			, BoundingboxExtentMin(InBoundingboxExtentMin)
 			, BoundingboxExtentMax(InBoundingboxExtentMax)
-			, SurfaceType(InSurfaceType)
+			, SurfaceType1(InSurfaceType1)
+			, SurfaceType2(InSurfaceType2)
 		{}
 
 		FCollidingDataExt(const FCollidingData& InCollisionData)
@@ -130,7 +132,8 @@ namespace Chaos
 			, BoundingboxVolume((FReal)-1.0)
 			, BoundingboxExtentMin((FReal)-1.0)
 			, BoundingboxExtentMax((FReal)-1.0)
-			, SurfaceType(-1)
+			, SurfaceType1(-1)
+			, SurfaceType2(-1)
 		{
 		}
 
@@ -142,7 +145,8 @@ namespace Chaos
 		FReal Mass1, Mass2;
 		FReal BoundingboxVolume;
 		FReal BoundingboxExtentMin, BoundingboxExtentMax;
-		int32 SurfaceType;
+		int32 SurfaceType1, SurfaceType2;
+		FName PhysicalMaterialName1, PhysicalMaterialName2;
 	};
 
 	/*
