- Fix for LWC changes that broke GPU buffers
- Reduce memory allocations by removing FString
- Fix potential OOB array access

#rnx

diff --git a/Engine/Plugins/Experimental/ChaosNiagara/Source/ChaosNiagara/Classes/NiagaraDataInterfaceChaosDestruction.h b/Engine/Plugins/Experimental/ChaosNiagara/Source/ChaosNiagara/Classes/NiagaraDataInterfaceChaosDestruction.h
--- a/Engine/Plugins/Experimental/ChaosNiagara/Source/ChaosNiagara/Classes/NiagaraDataInterfaceChaosDestruction.h
+++ b/Engine/Plugins/Experimental/ChaosNiagara/Source/ChaosNiagara/Classes/NiagaraDataInterfaceChaosDestruction.h
@@ -123,8 +123,8 @@ struct FNDIChaosDestruction_InstanceData
 
 struct FNiagaraDIChaosDestruction_GPUData
 {
-	TArray<FVector> PositionArray;
-	TArray<FVector> VelocityArray;
+	TArray<FVector3f> PositionArray;
+	TArray<FVector3f> VelocityArray;
 	TArray<float> ExtentMinArray;
 	TArray<float> ExtentMaxArray;
 	TArray<float> VolumeArray;
@@ -216,8 +216,8 @@ struct FNiagaraDIChaosDestruction_GPUData
 
 struct FNiagaraDIChaosDestruction_InstanceDataToPassToRT
 {
-	TArray<FVector>* PositionArray;
-	TArray<FVector>* VelocityArray;
+	TArray<FVector3f>* PositionArray;
+	TArray<FVector3f>* VelocityArray;
 	TArray<float>* ExtentMinArray;
 	TArray<float>* ExtentMaxArray;
 	TArray<float>* VolumeArray;
diff --git a/Engine/Plugins/Experimental/ChaosNiagara/Source/ChaosNiagara/Private/NiagaraDataInterfaceChaosDestruction.cpp b/Engine/Plugins/Experimental/ChaosNiagara/Source/ChaosNiagara/Private/NiagaraDataInterfaceChaosDestruction.cpp
--- a/Engine/Plugins/Experimental/ChaosNiagara/Source/ChaosNiagara/Private/NiagaraDataInterfaceChaosDestruction.cpp
+++ b/Engine/Plugins/Experimental/ChaosNiagara/Source/ChaosNiagara/Private/NiagaraDataInterfaceChaosDestruction.cpp
@@ -2788,29 +2788,14 @@ void UNiagaraDataInterfaceChaosDestruction::GetPosition(FVectorVMExternalFunctio
 {
 	VectorVM::FUserPtrHandler<FNDIChaosDestruction_InstanceData> InstData(Context);
 	ParticleIDParamType ParticleIDParam(Context);
-	VectorVM::FExternalFuncRegisterHandler<float> OutSampleX(Context);
-	VectorVM::FExternalFuncRegisterHandler<float> OutSampleY(Context);
-	VectorVM::FExternalFuncRegisterHandler<float> OutSampleZ(Context);
+	FNDIOutputParam<FVector3f> OutSample(Context);
 
 	for (int32 i = 0; i < Context.GetNumInstances(); ++i)
 	{
-		//		if (InstData->ParticleDataArray.Num())
-		if (InstData->PositionArray.Num())
-		{
-			int32 ParticleID = ParticleIDParam.Get();
-			// Remap ParticleID
-			ParticleID -= LastSpawnedPointID + 1;
-			FVector V = InstData->PositionArray[ParticleID];
-			//UE_LOG(LogScript, Warning, TEXT("Position: %d: %f %f %f"), ParticleID, V.X, V.Y, V.Z);
-			*OutSampleX.GetDest() = V.X;
-			*OutSampleY.GetDest() = V.Y;
-			*OutSampleZ.GetDest() = V.Z;
-		}
-
-		ParticleIDParam.Advance();
-		OutSampleX.Advance();
-		OutSampleY.Advance();
-		OutSampleZ.Advance();
+		const int32 ParticleID = ParticleIDParam.GetAndAdvance() - LastSpawnedPointID + 1;
+		//-TODO: LWC
+		const FVector3f Sample = InstData->PositionArray.IsValidIndex(ParticleID) ? (FVector3f)InstData->PositionArray[ParticleID] : FVector3f::Zero();
+		OutSample.SetAndAdvance(Sample);
 	}
 }
 
@@ -2819,28 +2804,13 @@ void UNiagaraDataInterfaceChaosDestruction::GetNormal(FVectorVMExternalFunctionC
 {
 	VectorVM::FUserPtrHandler<FNDIChaosDestruction_InstanceData> InstData(Context);
 	ParticleIDParamType ParticleIDParam(Context);
-	VectorVM::FExternalFuncRegisterHandler<float> OutSampleX(Context);
-	VectorVM::FExternalFuncRegisterHandler<float> OutSampleY(Context);
-	VectorVM::FExternalFuncRegisterHandler<float> OutSampleZ(Context);
+	FNDIOutputParam<FVector3f> OutSample(Context);
 
 	for (int32 i = 0; i < Context.GetNumInstances(); ++i)
 	{
-		if (InstData->IncomingNormalArray.Num())
-		{
-			int32 ParticleID = ParticleIDParam.Get();
-			// Remap ParticleID
-			ParticleID -= LastSpawnedPointID + 1;
-			FVector V = InstData->IncomingNormalArray[ParticleID];
-
-			*OutSampleX.GetDest() = V.X;
-			*OutSampleY.GetDest() = V.Y;
-			*OutSampleZ.GetDest() = V.Z;
-		}
-
-		ParticleIDParam.Advance();
-		OutSampleX.Advance();
-		OutSampleY.Advance();
-		OutSampleZ.Advance();
+		const int32 ParticleID = ParticleIDParam.GetAndAdvance() - LastSpawnedPointID + 1;
+		const FVector3f Sample = InstData->IncomingNormalArray.IsValidIndex(ParticleID) ? (FVector3f)InstData->IncomingNormalArray[ParticleID] : FVector3f::Zero();
+		OutSample.SetAndAdvance(Sample);
 	}
 }
 
@@ -2849,28 +2819,13 @@ void UNiagaraDataInterfaceChaosDestruction::GetVelocity(FVectorVMExternalFunctio
 {
 	VectorVM::FUserPtrHandler<FNDIChaosDestruction_InstanceData> InstData(Context);
 	ParticleIDParamType ParticleIDParam(Context);
-	VectorVM::FExternalFuncRegisterHandler<float> OutSampleX(Context);
-	VectorVM::FExternalFuncRegisterHandler<float> OutSampleY(Context);
-	VectorVM::FExternalFuncRegisterHandler<float> OutSampleZ(Context);
+	FNDIOutputParam<FVector3f> OutSample(Context);
 
 	for (int32 i = 0; i < Context.GetNumInstances(); ++i)
 	{
-		//		if (InstData->ParticleDataArray.Num())
-		if (InstData->VelocityArray.Num())
-		{
-			int32 ParticleID = ParticleIDParam.Get();
-			// Remap ParticleID
-			ParticleID -= LastSpawnedPointID + 1;
-			FVector V = InstData->VelocityArray[ParticleID];
-			*OutSampleX.GetDest() = V.X;
-			*OutSampleY.GetDest() = V.Y;
-			*OutSampleZ.GetDest() = V.Z;
-		}
-
-		ParticleIDParam.Advance();
-		OutSampleX.Advance();
-		OutSampleY.Advance();
-		OutSampleZ.Advance();
+		const int32 ParticleID = ParticleIDParam.GetAndAdvance() - LastSpawnedPointID + 1;
+		const FVector3f Sample = InstData->VelocityArray.IsValidIndex(ParticleID) ? (FVector3f)InstData->VelocityArray[ParticleID] : FVector3f::Zero();
+		OutSample.SetAndAdvance(Sample);
 	}
 }
 
@@ -2879,28 +2834,13 @@ void UNiagaraDataInterfaceChaosDestruction::GetAngularVelocity(FVectorVMExternal
 {
 	VectorVM::FUserPtrHandler<FNDIChaosDestruction_InstanceData> InstData(Context);
 	ParticleIDParamType ParticleIDParam(Context);
-	VectorVM::FExternalFuncRegisterHandler<float> OutSampleX(Context);
-	VectorVM::FExternalFuncRegisterHandler<float> OutSampleY(Context);
-	VectorVM::FExternalFuncRegisterHandler<float> OutSampleZ(Context);
+	FNDIOutputParam<FVector3f> OutSample(Context);
 
 	for (int32 i = 0; i < Context.GetNumInstances(); ++i)
 	{
-		if (InstData->IncomingAngularVelocity1Array.Num())
-		{
-			int32 ParticleID = ParticleIDParam.Get();
-			// Remap ParticleID
-			ParticleID -= LastSpawnedPointID + 1;
-			FVector W = InstData->IncomingAngularVelocity1Array[ParticleID];
-
-			*OutSampleX.GetDest() = W.X;
-			*OutSampleY.GetDest() = W.Y;
-			*OutSampleZ.GetDest() = W.Z;
-		}
-
-		ParticleIDParam.Advance();
-		OutSampleX.Advance();
-		OutSampleY.Advance();
-		OutSampleZ.Advance();
+		const int32 ParticleID = ParticleIDParam.GetAndAdvance() - LastSpawnedPointID + 1;
+		const FVector3f Sample = InstData->IncomingAngularVelocity1Array.IsValidIndex(ParticleID) ? (FVector3f)InstData->IncomingAngularVelocity1Array[ParticleID] : FVector3f::Zero();
+		OutSample.SetAndAdvance(Sample);
 	}
 }
 
@@ -2909,22 +2849,13 @@ void UNiagaraDataInterfaceChaosDestruction::GetExtentMin(FVectorVMExternalFuncti
 {
 	VectorVM::FUserPtrHandler<FNDIChaosDestruction_InstanceData> InstData(Context);
 	ParticleIDParamType ParticleIDParam(Context);
-	VectorVM::FExternalFuncRegisterHandler<float> OutValue(Context);
+	FNDIOutputParam<float> OutSample(Context);
 
 	for (int32 i = 0; i < Context.GetNumInstances(); ++i)
 	{
-		if (InstData->ExtentMinArray.Num())
-		{
-			int32 ParticleID = ParticleIDParam.Get();
-			// Remap ParticleID
-			ParticleID -= LastSpawnedPointID + 1;
-			float Value = InstData->ExtentMinArray[ParticleID];
-
-			*OutValue.GetDest() = Value;
-		}
-
-		ParticleIDParam.Advance();
-		OutValue.Advance();
+		const int32 ParticleID = ParticleIDParam.GetAndAdvance() - LastSpawnedPointID + 1;
+		const float Sample = InstData->ExtentMinArray.IsValidIndex(ParticleID) ? InstData->ExtentMinArray[ParticleID] : 0.0f;
+		OutSample.SetAndAdvance(Sample);
 	}
 }
 
@@ -2933,22 +2864,13 @@ void UNiagaraDataInterfaceChaosDestruction::GetExtentMax(FVectorVMExternalFuncti
 {
 	VectorVM::FUserPtrHandler<FNDIChaosDestruction_InstanceData> InstData(Context);
 	ParticleIDParamType ParticleIDParam(Context);
-	VectorVM::FExternalFuncRegisterHandler<float> OutValue(Context);
+	FNDIOutputParam<float> OutSample(Context);
 
 	for (int32 i = 0; i < Context.GetNumInstances(); ++i)
 	{
-		if (InstData->ExtentMaxArray.Num())
-		{
-			int32 ParticleID = ParticleIDParam.Get();
-			// Remap ParticleID
-			ParticleID -= LastSpawnedPointID + 1;
-			float Value = InstData->ExtentMaxArray[ParticleID];
-
-			*OutValue.GetDest() = Value;
-		}
-
-		ParticleIDParam.Advance();
-		OutValue.Advance();
+		const int32 ParticleID = ParticleIDParam.GetAndAdvance() - LastSpawnedPointID + 1;
+		const float Sample = InstData->ExtentMaxArray.IsValidIndex(ParticleID) ? InstData->ExtentMaxArray[ParticleID] : 0.0f;
+		OutSample.SetAndAdvance(Sample);
 	}
 }
 
@@ -2957,22 +2879,13 @@ void UNiagaraDataInterfaceChaosDestruction::GetVolume(FVectorVMExternalFunctionC
 {
 	VectorVM::FUserPtrHandler<FNDIChaosDestruction_InstanceData> InstData(Context);
 	ParticleIDParamType ParticleIDParam(Context);
-	VectorVM::FExternalFuncRegisterHandler<float> OutValue(Context);
+	FNDIOutputParam<float> OutSample(Context);
 
 	for (int32 i = 0; i < Context.GetNumInstances(); ++i)
 	{
-		if (InstData->VolumeArray.Num())
-		{
-			int32 ParticleID = ParticleIDParam.Get();
-			// Remap ParticleID
-			ParticleID -= LastSpawnedPointID + 1;
-			float Value = InstData->VolumeArray[ParticleID];
-
-			*OutValue.GetDest() = Value;
-		}
-
-		ParticleIDParam.Advance();
-		OutValue.Advance();
+		const int32 ParticleID = ParticleIDParam.GetAndAdvance() - LastSpawnedPointID + 1;
+		const float Sample = InstData->VolumeArray.IsValidIndex(ParticleID) ? InstData->VolumeArray[ParticleID] : 0.0f;
+		OutSample.SetAndAdvance(Sample);
 	}
 }
 
@@ -2981,29 +2894,26 @@ void UNiagaraDataInterfaceChaosDestruction::GetParticleIdsToSpawnAtTime(FVectorV
 {
 	VectorVM::FUserPtrHandler<FNDIChaosDestruction_InstanceData> InstData(Context);
 	TimeParamType TimeParam(Context);
-	VectorVM::FExternalFuncRegisterHandler<int32> OutMinValue(Context);
-	VectorVM::FExternalFuncRegisterHandler<int32> OutMaxValue(Context);
-	VectorVM::FExternalFuncRegisterHandler<int32> OutCountValue(Context);
+	FNDIOutputParam<int32> OutMinValue(Context);
+	FNDIOutputParam<int32> OutMaxValue(Context);
+	FNDIOutputParam<int32> OutCountValue(Context);
 
-	if (DoSpawn && ShouldSpawn && InstData->PositionArray.Num())
-	{
-		*OutMinValue.GetDest() = LastSpawnedPointID + 1;
-		*OutMaxValue.GetDest() = LastSpawnedPointID + InstData->PositionArray.Num();
-		*OutCountValue.GetDest() = InstData->PositionArray.Num();
-		//UE_LOG(LogScript, Warning, TEXT("Min = %d, Max = %d, Count = %d"), LastSpawnedPointID + 1, LastSpawnedPointID + InstData->PositionArray.Num(), InstData->PositionArray.Num());
-	} 
-	else 
+	for (int32 i = 0; i < Context.GetNumInstances(); ++i)
 	{
-		*OutMinValue.GetDest() = 0;
-		*OutMaxValue.GetDest() = 0;
-		*OutCountValue.GetDest() = 0;		
+		if (DoSpawn && ShouldSpawn && InstData->PositionArray.Num())
+		{
+			OutMinValue.SetAndAdvance(LastSpawnedPointID + 1);
+			OutMaxValue.SetAndAdvance(LastSpawnedPointID + InstData->PositionArray.Num());
+			OutCountValue.SetAndAdvance(InstData->PositionArray.Num());
+			//UE_LOG(LogScript, Warning, TEXT("Min = %d, Max = %d, Count = %d"), LastSpawnedPointID + 1, LastSpawnedPointID + InstData->PositionArray.Num(), InstData->PositionArray.Num());
+		}
+		else
+		{
+			OutMinValue.SetAndAdvance(0);
+			OutMaxValue.SetAndAdvance(0);
+			OutCountValue.SetAndAdvance(0);
+		}
 	}
-	
-
-	TimeParam.Advance();
-	OutMinValue.Advance();
-	OutMaxValue.Advance();
-	OutCountValue.Advance();
 }
 
 template<typename ParticleIDParamType>
@@ -3011,23 +2921,24 @@ void UNiagaraDataInterfaceChaosDestruction::GetPointType(FVectorVMExternalFuncti
 {
 	VectorVM::FUserPtrHandler<FNDIChaosDestruction_InstanceData> InstData(Context);
 	ParticleIDParamType ParticleIDParam(Context);
-	VectorVM::FExternalFuncRegisterHandler<int32> OutValue(Context);
+	FNDIOutputParam<int32> OutSample(Context);
 
 	for (int32 i = 0; i < Context.GetNumInstances(); ++i)
 	{
-		if (InstData->PositionArray.Num())
-		{
-			int32 ParticleID = ParticleIDParam.Get();
-			// Remap ParticleID
-			ParticleID -= LastSpawnedPointID + 1;
+		//if (InstData->PositionArray.Num())
+		//{
+		//	int32 ParticleID = ParticleIDParam.Get();
+		//	// Remap ParticleID
+		//	ParticleID -= LastSpawnedPointID + 1;
 
-			int32 Value = 0;
+		//	int32 Value = 0;
 
-			*OutValue.GetDest() = Value;
-		}
+		//	*OutValue.GetDest() = Value;
+		//}
 
-		ParticleIDParam.Advance();
-		OutValue.Advance();
+		//ParticleIDParam.Advance();
+		//OutValue.Advance();
+		OutSample.SetAndAdvance(0);
 	}
 }
 
@@ -3036,31 +2947,13 @@ void UNiagaraDataInterfaceChaosDestruction::GetColor(FVectorVMExternalFunctionCo
 {
 	VectorVM::FUserPtrHandler<FNDIChaosDestruction_InstanceData> InstData(Context);
 	ParticleIDParamType ParticleIDParam(Context);
-	VectorVM::FExternalFuncRegisterHandler<float> OutSampleR(Context);
-	VectorVM::FExternalFuncRegisterHandler<float> OutSampleG(Context);
-	VectorVM::FExternalFuncRegisterHandler<float> OutSampleB(Context);
-	VectorVM::FExternalFuncRegisterHandler<float> OutSampleA(Context);
+	FNDIOutputParam<FLinearColor> OutSample(Context);
 
 	for (int32 i = 0; i < Context.GetNumInstances(); ++i)
 	{
-		if (InstData->ColorArray.Num())
-		{
-			int32 ParticleID = ParticleIDParam.Get();
-			// Remap ParticleID
-			ParticleID -= LastSpawnedPointID + 1;
-			FLinearColor V = InstData->ColorArray[ParticleID];
-
-			*OutSampleR.GetDest() = V.R;
-			*OutSampleG.GetDest() = V.G;
-			*OutSampleB.GetDest() = V.B;
-			*OutSampleA.GetDest() = V.A;
-		}
-
-		ParticleIDParam.Advance();
-		OutSampleR.Advance();
-		OutSampleG.Advance();
-		OutSampleB.Advance();
-		OutSampleA.Advance();
+		const int32 ParticleID = ParticleIDParam.GetAndAdvance() - LastSpawnedPointID + 1;
+		const FLinearColor Sample = InstData->ColorArray.IsValidIndex(ParticleID) ? InstData->ColorArray[ParticleID] : FLinearColor::White;
+		OutSample.SetAndAdvance(Sample);
 	}
 }
 
@@ -3068,11 +2961,12 @@ template<typename ParticleIDParamType>
 void UNiagaraDataInterfaceChaosDestruction::GetSolverTime(FVectorVMExternalFunctionContext& Context)
 {
 	VectorVM::FUserPtrHandler<FNDIChaosDestruction_InstanceData> InstData(Context);
-	VectorVM::FExternalFuncRegisterHandler<float> OutValue(Context);
-
-	*OutValue.GetDest() = SolverTime;
+	FNDIOutputParam<float> OutValue(Context);
 
-	OutValue.Advance();
+	for (int32 i = 0; i < Context.GetNumInstances(); ++i)
+	{
+		OutValue.SetAndAdvance(SolverTime);
+	}
 }
 
 template<typename ParticleIDParamType>
@@ -3080,22 +2974,13 @@ void UNiagaraDataInterfaceChaosDestruction::GetDensity(FVectorVMExternalFunction
 {
 	VectorVM::FUserPtrHandler<FNDIChaosDestruction_InstanceData> InstData(Context);
 	ParticleIDParamType ParticleIDParam(Context);
-	VectorVM::FExternalFuncRegisterHandler<float> OutValue(Context);
+	FNDIOutputParam<float> OutSample(Context);
 
 	for (int32 i = 0; i < Context.GetNumInstances(); ++i)
 	{
-		if (InstData->DensityArray.Num())
-		{
-			int32 ParticleID = ParticleIDParam.Get();
-			// Remap ParticleID
-			ParticleID -= LastSpawnedPointID + 1;
-			float Value = InstData->DensityArray[ParticleID];
-
-			*OutValue.GetDest() = Value;
-		}
-
-		ParticleIDParam.Advance();
-		OutValue.Advance();
+		const int32 ParticleID = ParticleIDParam.GetAndAdvance() - LastSpawnedPointID + 1;
+		const float Sample = InstData->DensityArray.IsValidIndex(ParticleID) ? InstData->DensityArray[ParticleID] : 0.0f;
+		OutSample.SetAndAdvance(Sample);
 	}
 }
 
@@ -3104,22 +2989,13 @@ void UNiagaraDataInterfaceChaosDestruction::GetFriction(FVectorVMExternalFunctio
 {
 	VectorVM::FUserPtrHandler<FNDIChaosDestruction_InstanceData> InstData(Context);
 	ParticleIDParamType ParticleIDParam(Context);
-	VectorVM::FExternalFuncRegisterHandler<float> OutValue(Context);
+	FNDIOutputParam<float> OutSample(Context);
 
 	for (int32 i = 0; i < Context.GetNumInstances(); ++i)
 	{
-		if (InstData->FrictionArray.Num())
-		{
-			int32 ParticleID = ParticleIDParam.Get();
-			// Remap ParticleID
-			ParticleID -= LastSpawnedPointID + 1;
-			float Value = InstData->FrictionArray[ParticleID];
-
-			*OutValue.GetDest() = Value;
-		}
-
-		ParticleIDParam.Advance();
-		OutValue.Advance();
+		const int32 ParticleID = ParticleIDParam.GetAndAdvance() - LastSpawnedPointID + 1;
+		const float Sample = InstData->FrictionArray.IsValidIndex(ParticleID) ? InstData->FrictionArray[ParticleID] : 0.0f;
+		OutSample.SetAndAdvance(Sample);
 	}
 }
 
@@ -3128,22 +3004,13 @@ void UNiagaraDataInterfaceChaosDestruction::GetRestitution(FVectorVMExternalFunc
 {
 	VectorVM::FUserPtrHandler<FNDIChaosDestruction_InstanceData> InstData(Context);
 	ParticleIDParamType ParticleIDParam(Context);
-	VectorVM::FExternalFuncRegisterHandler<float> OutValue(Context);
+	FNDIOutputParam<float> OutSample(Context);
 
 	for (int32 i = 0; i < Context.GetNumInstances(); ++i)
 	{
-		if (InstData->RestitutionArray.Num())
-		{
-			int32 ParticleID = ParticleIDParam.Get();
-			// Remap ParticleID
-			ParticleID -= LastSpawnedPointID + 1;
-			float Value = InstData->RestitutionArray[ParticleID];
-
-			*OutValue.GetDest() = Value;
-		}
-
-		ParticleIDParam.Advance();
-		OutValue.Advance();
+		const int32 ParticleID = ParticleIDParam.GetAndAdvance() - LastSpawnedPointID + 1;
+		const float Sample = InstData->RestitutionArray.IsValidIndex(ParticleID) ? InstData->RestitutionArray[ParticleID] : 0.0f;
+		OutSample.SetAndAdvance(Sample);
 	}
 }
 
@@ -3152,55 +3019,26 @@ void UNiagaraDataInterfaceChaosDestruction::GetTransform(FVectorVMExternalFuncti
 {
 	VectorVM::FUserPtrHandler<FNDIChaosDestruction_InstanceData> InstData(Context);
 	ParticleIDParamType ParticleIDParam(Context);
-	VectorVM::FExternalFuncRegisterHandler<float> OutTranslationX(Context);
-	VectorVM::FExternalFuncRegisterHandler<float> OutTranslationY(Context);
-	VectorVM::FExternalFuncRegisterHandler<float> OutTranslationZ(Context);
-	VectorVM::FExternalFuncRegisterHandler<float> OutRotationX(Context);
-	VectorVM::FExternalFuncRegisterHandler<float> OutRotationY(Context);
-	VectorVM::FExternalFuncRegisterHandler<float> OutRotationZ(Context);
-	VectorVM::FExternalFuncRegisterHandler<float> OutRotationW(Context);
-	VectorVM::FExternalFuncRegisterHandler<float> OutScaleX(Context);
-	VectorVM::FExternalFuncRegisterHandler<float> OutScaleY(Context);
-	VectorVM::FExternalFuncRegisterHandler<float> OutScaleZ(Context);
+	FNDIOutputParam<FVector3f> OutTranslation(Context);
+	FNDIOutputParam<FQuat4f> OutRotation(Context);
+	FNDIOutputParam<FVector3f> OutScale(Context);
 
 	for (int32 i = 0; i < Context.GetNumInstances(); ++i)
 	{
-		if (InstData->TransformTranslationArray.Num())
+		const int32 ParticleID = ParticleIDParam.GetAndAdvance() - LastSpawnedPointID + 1;
+		if (InstData->TransformTranslationArray.IsValidIndex(ParticleID))
 		{
-			int32 ParticleID = ParticleIDParam.Get();
-			// Remap ParticleID
-			ParticleID -= LastSpawnedPointID + 1;
-			FVector Translation = InstData->TransformTranslationArray[ParticleID];
-			FQuat Rotation = InstData->TransformRotationArray[ParticleID];
-			FVector Scale = InstData->TransformScaleArray[ParticleID];
-			//UE_LOG(LogScript, Warning, TEXT("%f %f %f    %f %f %f %f   %f %f %f"), Translation.X, Translation.Y, Translation.Z, Rotation.X, Rotation.Y, Rotation.Z, Rotation.W, Scale.X, Scale.X, Scale.X)
-			//UE_LOG(LogScript, Warning, TEXT("Transform: %f %f %f"), Translation.X, Translation.Y, Translation.Z)
-
-			*OutTranslationX.GetDest() = Translation.X;
-			*OutTranslationY.GetDest() = Translation.Y;
-			*OutTranslationZ.GetDest() = Translation.Z;
-			*OutRotationX.GetDest() = Rotation.X;
-			*OutRotationY.GetDest() = Rotation.Y;
-			*OutRotationZ.GetDest() = Rotation.Z;
-			*OutRotationW.GetDest() = Rotation.W;
-			*OutScaleX.GetDest() = Scale.X;
-			*OutScaleY.GetDest() = Scale.Y;
-			*OutScaleZ.GetDest() = Scale.Z;
+			//-TODO: LWC
+			OutTranslation.SetAndAdvance((FVector3f)InstData->TransformTranslationArray[ParticleID]);
+			OutRotation.SetAndAdvance((FQuat4f)InstData->TransformRotationArray[ParticleID]);
+			OutScale.SetAndAdvance((FVector3f)InstData->TransformScaleArray[ParticleID]);
+		}
+		else
+		{
+			OutTranslation.SetAndAdvance(FVector3f::Zero());
+			OutRotation.SetAndAdvance(FQuat4f::Identity);
+			OutScale.SetAndAdvance(FVector3f(1.0f, 1.0f, 1.0f));
 		}
-
-		ParticleIDParam.Advance();
-		OutTranslationX.Advance();
-		OutTranslationY.Advance();
-		OutTranslationZ.Advance();
-		OutRotationX.Advance();
-		OutRotationY.Advance();
-		OutRotationZ.Advance();
-		OutRotationW.Advance();
-		OutScaleX.Advance();
-		OutScaleY.Advance();
-		OutScaleZ.Advance();
-		//OutRotation.Advance();
-		//OutScale.Advance();
 	}
 }
 
@@ -3209,28 +3047,13 @@ void UNiagaraDataInterfaceChaosDestruction::GetSize(FVectorVMExternalFunctionCon
 {
 	VectorVM::FUserPtrHandler<FNDIChaosDestruction_InstanceData> InstData(Context);
 	ParticleIDParamType ParticleIDParam(Context);
-	VectorVM::FExternalFuncRegisterHandler<float> OutSizeX(Context);
-	VectorVM::FExternalFuncRegisterHandler<float> OutSizeY(Context);
-	VectorVM::FExternalFuncRegisterHandler<float> OutSizeZ(Context);
+	FNDIOutputParam<FVector3f> OutSize(Context);
 
 	for (int32 i = 0; i < Context.GetNumInstances(); ++i)
 	{
-		if (InstData->TransformTranslationArray.Num())
-		{
-			int32 ParticleID = ParticleIDParam.Get();
-			// Remap ParticleID
-			ParticleID -= LastSpawnedPointID + 1;
-			FVector Bounds = InstData->BoundsArray[ParticleID];
-
-			*OutSizeX.GetDest() = Bounds.X;
-			*OutSizeY.GetDest() = Bounds.Y;
-			*OutSizeZ.GetDest() = Bounds.Z;
-		}
-
-		ParticleIDParam.Advance();
-		OutSizeX.Advance();
-		OutSizeY.Advance();
-		OutSizeZ.Advance();
+		const int32 ParticleID = ParticleIDParam.GetAndAdvance() - LastSpawnedPointID + 1;
+		const FVector3f Size = InstData->BoundsArray.IsValidIndex(ParticleID) ? (FVector3f)InstData->BoundsArray[ParticleID] : FVector3f::Zero();
+		OutSize.SetAndAdvance(Size);
 	}
 }
 
@@ -3239,22 +3062,13 @@ void UNiagaraDataInterfaceChaosDestruction::GetSurfaceType(FVectorVMExternalFunc
 {
 	VectorVM::FUserPtrHandler<FNDIChaosDestruction_InstanceData> InstData(Context);
 	ParticleIDParamType ParticleIDParam(Context);
-	VectorVM::FExternalFuncRegisterHandler<int32> OutValue(Context);
+	FNDIOutputParam<int32> OutValue(Context);
 
 	for (int32 i = 0; i < Context.GetNumInstances(); ++i)
 	{
-		if (InstData->SurfaceTypeArray.Num())
-		{
-			int32 ParticleID = ParticleIDParam.Get();
-			// Remap ParticleID
-			ParticleID -= LastSpawnedPointID + 1;
-			int32 Value = InstData->SurfaceTypeArray[ParticleID];
-
-			*OutValue.GetDest() = Value;
-		}
-
-		ParticleIDParam.Advance();
-		OutValue.Advance();
+		const int32 ParticleID = ParticleIDParam.GetAndAdvance() - LastSpawnedPointID + 1;
+		const int32 Value = InstData->SurfaceTypeArray.IsValidIndex(ParticleID) ? InstData->SurfaceTypeArray[ParticleID] : 0;
+		OutValue.SetAndAdvance(Value);
 	}
 }
 
@@ -3263,112 +3077,32 @@ void UNiagaraDataInterfaceChaosDestruction::GetCollisionData(FVectorVMExternalFu
 {
 	VectorVM::FUserPtrHandler<FNDIChaosDestruction_InstanceData> InstData(Context);
 	ParticleIDParamType ParticleIDParam(Context);
-	VectorVM::FExternalFuncRegisterHandler<float> OutLocationX(Context);
-	VectorVM::FExternalFuncRegisterHandler<float> OutLocationY(Context);
-	VectorVM::FExternalFuncRegisterHandler<float> OutLocationZ(Context);
-	VectorVM::FExternalFuncRegisterHandler<float> OutAccumulatedImpulseX(Context);
-	VectorVM::FExternalFuncRegisterHandler<float> OutAccumulatedImpulseY(Context);
-	VectorVM::FExternalFuncRegisterHandler<float> OutAccumulatedImpulseZ(Context);
-	VectorVM::FExternalFuncRegisterHandler<float> OutNormalX(Context);
-	VectorVM::FExternalFuncRegisterHandler<float> OutNormalY(Context);
-	VectorVM::FExternalFuncRegisterHandler<float> OutNormalZ(Context);
-	VectorVM::FExternalFuncRegisterHandler<float> OutVelocity1X(Context);
-	VectorVM::FExternalFuncRegisterHandler<float> OutVelocity1Y(Context);
-	VectorVM::FExternalFuncRegisterHandler<float> OutVelocity1Z(Context);
-	VectorVM::FExternalFuncRegisterHandler<float> OutVelocity2X(Context);
-	VectorVM::FExternalFuncRegisterHandler<float> OutVelocity2Y(Context);
-	VectorVM::FExternalFuncRegisterHandler<float> OutVelocity2Z(Context);
-	VectorVM::FExternalFuncRegisterHandler<float> OutAngularVelocity1X(Context);
-	VectorVM::FExternalFuncRegisterHandler<float> OutAngularVelocity1Y(Context);
-	VectorVM::FExternalFuncRegisterHandler<float> OutAngularVelocity1Z(Context);
-	VectorVM::FExternalFuncRegisterHandler<float> OutAngularVelocity2X(Context);
-	VectorVM::FExternalFuncRegisterHandler<float> OutAngularVelocity2Y(Context);
-	VectorVM::FExternalFuncRegisterHandler<float> OutAngularVelocity2Z(Context);
-	VectorVM::FExternalFuncRegisterHandler<float> OutMass1(Context);
-	VectorVM::FExternalFuncRegisterHandler<float> OutMass2(Context);
-	VectorVM::FExternalFuncRegisterHandler<float> OutTime(Context);
+	FNDIOutputParam<FVector3f> OutLocation(Context);
+	FNDIOutputParam<FVector3f> OutAccumulatedImpulse(Context);
+	FNDIOutputParam<FVector3f> OutNormal(Context);
+	FNDIOutputParam<FVector3f> OutVelocity1(Context);
+	FNDIOutputParam<FVector3f> OutVelocity2(Context);
+	FNDIOutputParam<FVector3f> OutAngularVelocity1(Context);
+	FNDIOutputParam<FVector3f> OutAngularVelocity2(Context);
+	FNDIOutputParam<float> OutMass1(Context);
+	FNDIOutputParam<float> OutMass2(Context);
+	FNDIOutputParam<float> OutTime(Context);
 
 	for (int32 i = 0; i < Context.GetNumInstances(); ++i)
 	{
-		if (InstData->IncomingLocationArray.Num())
-		{
-			int32 ParticleID = ParticleIDParam.Get();
-			// Remap ParticleID
-			ParticleID -= LastSpawnedPointID + 1;
-
-			FVector VValue;
-			float FValue;
-
-			VValue = InstData->IncomingLocationArray[ParticleID];
-			*OutLocationX.GetDest() = VValue.X;
-			*OutLocationY.GetDest() = VValue.Y;
-			*OutLocationZ.GetDest() = VValue.Z;
-
-			VValue = InstData->IncomingAccumulatedImpulseArray[ParticleID];
-			*OutAccumulatedImpulseX.GetDest() = VValue.X;
-			*OutAccumulatedImpulseY.GetDest() = VValue.Y;
-			*OutAccumulatedImpulseZ.GetDest() = VValue.Z;
-
-			VValue = InstData->IncomingNormalArray[ParticleID];
-			*OutNormalX.GetDest() = VValue.X;
-			*OutNormalY.GetDest() = VValue.Y;
-			*OutNormalZ.GetDest() = VValue.Z;
-
-			VValue = InstData->IncomingVelocity1Array[ParticleID];
-			*OutVelocity1X.GetDest() = VValue.X;
-			*OutVelocity1Y.GetDest() = VValue.Y;
-			*OutVelocity1Z.GetDest() = VValue.Z;
-
-			VValue = InstData->IncomingVelocity2Array[ParticleID];
-			*OutVelocity2X.GetDest() = VValue.X;
-			*OutVelocity2Y.GetDest() = VValue.Y;
-			*OutVelocity2Z.GetDest() = VValue.Z;
-
-			VValue = InstData->IncomingAngularVelocity1Array[ParticleID];
-			*OutAngularVelocity1X.GetDest() = VValue.X;
-			*OutAngularVelocity1Y.GetDest() = VValue.Y;
-			*OutAngularVelocity1Z.GetDest() = VValue.Z;
-
-			VValue = InstData->IncomingAngularVelocity2Array[ParticleID];
-			*OutAngularVelocity2X.GetDest() = VValue.X;
-			*OutAngularVelocity2Y.GetDest() = VValue.Y;
-			*OutAngularVelocity2Z.GetDest() = VValue.Z;
-
-			FValue = InstData->IncomingMass1Array[ParticleID];
-			*OutMass1.GetDest() = FValue;
-
-			FValue = InstData->IncomingMass2Array[ParticleID];
-			*OutMass2.GetDest() = FValue;
-
-			FValue = InstData->IncomingTimeArray[ParticleID];
-			*OutTime.GetDest() = FValue;
-		}
+		const int32 ParticleID = ParticleIDParam.GetAndAdvance() - LastSpawnedPointID + 1;
 
-		ParticleIDParam.Advance();
-		OutLocationX.Advance();
-		OutLocationY.Advance();
-		OutLocationZ.Advance();
-		OutAccumulatedImpulseX.Advance();
-		OutAccumulatedImpulseY.Advance();
-		OutAccumulatedImpulseZ.Advance();
-		OutNormalX.Advance();
-		OutNormalY.Advance();
-		OutNormalZ.Advance();
-		OutVelocity1X.Advance();
-		OutVelocity1Y.Advance();
-		OutVelocity1Z.Advance();
-		OutVelocity2X.Advance();
-		OutVelocity2Y.Advance();
-		OutVelocity2Z.Advance();
-		OutAngularVelocity1X.Advance();
-		OutAngularVelocity1Y.Advance();
-		OutAngularVelocity1Z.Advance();
-		OutAngularVelocity2X.Advance();
-		OutAngularVelocity2Y.Advance();
-		OutAngularVelocity2Z.Advance();
-		OutMass1.Advance();
-		OutMass2.Advance();
-		OutTime.Advance();
+		//-TODO: LWC
+		OutLocation.SetAndAdvance(InstData->IncomingLocationArray.IsValidIndex(ParticleID) ? (FVector3f)InstData->IncomingLocationArray[ParticleID] : FVector3f::Zero());
+		OutAccumulatedImpulse.SetAndAdvance(InstData->IncomingAccumulatedImpulseArray.IsValidIndex(ParticleID) ? (FVector3f)InstData->IncomingAccumulatedImpulseArray[ParticleID] : FVector3f::Zero());
+		OutNormal.SetAndAdvance(InstData->IncomingNormalArray.IsValidIndex(ParticleID) ? (FVector3f)InstData->IncomingNormalArray[ParticleID] : FVector3f::Zero());
+		OutVelocity1.SetAndAdvance(InstData->IncomingVelocity1Array.IsValidIndex(ParticleID) ? (FVector3f)InstData->IncomingVelocity1Array[ParticleID] : FVector3f::Zero());
+		OutVelocity2.SetAndAdvance(InstData->IncomingVelocity2Array.IsValidIndex(ParticleID) ? (FVector3f)InstData->IncomingVelocity2Array[ParticleID] : FVector3f::Zero());
+		OutAngularVelocity1.SetAndAdvance(InstData->IncomingAngularVelocity1Array.IsValidIndex(ParticleID) ? (FVector3f)InstData->IncomingAngularVelocity1Array[ParticleID] : FVector3f::Zero());
+		OutAngularVelocity2.SetAndAdvance(InstData->IncomingAngularVelocity2Array.IsValidIndex(ParticleID) ? (FVector3f)InstData->IncomingAngularVelocity2Array[ParticleID] : FVector3f::Zero());
+		OutMass1.SetAndAdvance(InstData->IncomingMass1Array.IsValidIndex(ParticleID) ? InstData->IncomingMass1Array[ParticleID] : 0.0f);
+		OutMass2.SetAndAdvance(InstData->IncomingMass2Array.IsValidIndex(ParticleID) ? InstData->IncomingMass2Array[ParticleID] : 0.0f);
+		OutTime.SetAndAdvance(InstData->IncomingTimeArray.IsValidIndex(ParticleID) ? InstData->IncomingTimeArray[ParticleID] : 0.0f);
 	}
 }
 
@@ -3377,63 +3111,21 @@ void UNiagaraDataInterfaceChaosDestruction::GetBreakingData(FVectorVMExternalFun
 {
 	VectorVM::FUserPtrHandler<FNDIChaosDestruction_InstanceData> InstData(Context);
 	ParticleIDParamType ParticleIDParam(Context);
-	VectorVM::FExternalFuncRegisterHandler<float> OutLocationX(Context);
-	VectorVM::FExternalFuncRegisterHandler<float> OutLocationY(Context);
-	VectorVM::FExternalFuncRegisterHandler<float> OutLocationZ(Context);
-	VectorVM::FExternalFuncRegisterHandler<float> OutVelocityX(Context);
-	VectorVM::FExternalFuncRegisterHandler<float> OutVelocityY(Context);
-	VectorVM::FExternalFuncRegisterHandler<float> OutVelocityZ(Context);
-	VectorVM::FExternalFuncRegisterHandler<float> OutAngularVelocityX(Context);
-	VectorVM::FExternalFuncRegisterHandler<float> OutAngularVelocityY(Context);
-	VectorVM::FExternalFuncRegisterHandler<float> OutAngularVelocityZ(Context);
-	VectorVM::FExternalFuncRegisterHandler<float> OutMass(Context);
-	VectorVM::FExternalFuncRegisterHandler<float> OutTime(Context);
+	FNDIOutputParam<FVector3f> OutLocation(Context);
+	FNDIOutputParam<FVector3f> OutVelocity(Context);
+	FNDIOutputParam<FVector3f> OutAngularVelocity(Context);
+	FNDIOutputParam<float> OutMass(Context);
+	FNDIOutputParam<float> OutTime(Context);
 
 	for (int32 i = 0; i < Context.GetNumInstances(); ++i)
 	{
-		if (InstData->IncomingLocationArray.Num())
-		{
-			int32 ParticleID = ParticleIDParam.Get();
-			// Remap ParticleID
-			ParticleID -= LastSpawnedPointID + 1;
-
-			FVector VValue;
-			float FValue;
-
-			VValue = InstData->IncomingLocationArray[ParticleID];
-			*OutLocationX.GetDest() = VValue.X;
-			*OutLocationY.GetDest() = VValue.Y;
-			*OutLocationZ.GetDest() = VValue.Z;
-
-			VValue = InstData->IncomingVelocity1Array[ParticleID];
-			*OutVelocityX.GetDest() = VValue.X;
-			*OutVelocityY.GetDest() = VValue.Y;
-			*OutVelocityZ.GetDest() = VValue.Z;
-
-			VValue = InstData->IncomingAngularVelocity1Array[ParticleID];
-			*OutAngularVelocityX.GetDest() = VValue.X;
-			*OutAngularVelocityY.GetDest() = VValue.Y;
-			*OutAngularVelocityZ.GetDest() = VValue.Z;
+		const int32 ParticleID = ParticleIDParam.GetAndAdvance() - LastSpawnedPointID + 1;
 
-			FValue = InstData->IncomingMass1Array[ParticleID];
-			*OutMass.GetDest() = FValue;
-
-			FValue = InstData->IncomingTimeArray[ParticleID];
-			*OutTime.GetDest() = FValue;
-		}
-
-		ParticleIDParam.Advance();
-		OutLocationX.Advance();
-		OutLocationY.Advance();
-		OutLocationZ.Advance();
-		OutVelocityX.Advance();
-		OutVelocityY.Advance();
-		OutVelocityZ.Advance();
-		OutAngularVelocityX.Advance();
-		OutAngularVelocityY.Advance();
-		OutAngularVelocityZ.Advance();
-		OutMass.Advance();
-		OutTime.Advance();
+		OutLocation.SetAndAdvance(InstData->IncomingLocationArray.IsValidIndex(ParticleID) ? (FVector3f)InstData->IncomingLocationArray[ParticleID] : FVector3f::Zero());
+		OutVelocity.SetAndAdvance(InstData->IncomingVelocity1Array.IsValidIndex(ParticleID) ? (FVector3f)InstData->IncomingVelocity1Array[ParticleID] : FVector3f::Zero());
+		OutAngularVelocity.SetAndAdvance(InstData->IncomingAngularVelocity1Array.IsValidIndex(ParticleID) ? (FVector3f)InstData->IncomingAngularVelocity1Array[ParticleID] : FVector3f::Zero());
+		OutMass.SetAndAdvance(InstData->IncomingMass1Array.IsValidIndex(ParticleID) ? InstData->IncomingMass1Array[ParticleID] : 0.0f);
+		OutTime.SetAndAdvance(InstData->IncomingTimeArray.IsValidIndex(ParticleID) ? InstData->IncomingTimeArray[ParticleID] : 0.0f);
 	}
 }
 
@@ -3442,63 +3134,21 @@ void UNiagaraDataInterfaceChaosDestruction::GetTrailingData(FVectorVMExternalFun
 {
 	VectorVM::FUserPtrHandler<FNDIChaosDestruction_InstanceData> InstData(Context);
 	ParticleIDParamType ParticleIDParam(Context);
-	VectorVM::FExternalFuncRegisterHandler<float> OutLocationX(Context);
-	VectorVM::FExternalFuncRegisterHandler<float> OutLocationY(Context);
-	VectorVM::FExternalFuncRegisterHandler<float> OutLocationZ(Context);
-	VectorVM::FExternalFuncRegisterHandler<float> OutVelocityX(Context);
-	VectorVM::FExternalFuncRegisterHandler<float> OutVelocityY(Context);
-	VectorVM::FExternalFuncRegisterHandler<float> OutVelocityZ(Context);
-	VectorVM::FExternalFuncRegisterHandler<float> OutAngularVelocityX(Context);
-	VectorVM::FExternalFuncRegisterHandler<float> OutAngularVelocityY(Context);
-	VectorVM::FExternalFuncRegisterHandler<float> OutAngularVelocityZ(Context);
-	VectorVM::FExternalFuncRegisterHandler<float> OutMass(Context);
-	VectorVM::FExternalFuncRegisterHandler<float> OutTime(Context);
+	FNDIOutputParam<FVector3f> OutLocation(Context);
+	FNDIOutputParam<FVector3f> OutVelocity(Context);
+	FNDIOutputParam<FVector3f> OutAngularVelocity(Context);
+	FNDIOutputParam<float> OutMass(Context);
+	FNDIOutputParam<float> OutTime(Context);
 
 	for (int32 i = 0; i < Context.GetNumInstances(); ++i)
 	{
-		if (InstData->IncomingLocationArray.Num())
-		{
-			int32 ParticleID = ParticleIDParam.Get();
-			// Remap ParticleID
-			ParticleID -= LastSpawnedPointID + 1;
-
-			FVector VValue;
-			float FValue;
+		const int32 ParticleID = ParticleIDParam.GetAndAdvance() - LastSpawnedPointID + 1;
 
-			VValue = InstData->IncomingLocationArray[ParticleID];
-			*OutLocationX.GetDest() = VValue.X;
-			*OutLocationY.GetDest() = VValue.Y;
-			*OutLocationZ.GetDest() = VValue.Z;
-
-			VValue = InstData->IncomingVelocity1Array[ParticleID];
-			*OutVelocityX.GetDest() = VValue.X;
-			*OutVelocityY.GetDest() = VValue.Y;
-			*OutVelocityZ.GetDest() = VValue.Z;
-
-			VValue = InstData->IncomingAngularVelocity1Array[ParticleID];
-			*OutAngularVelocityX.GetDest() = VValue.X;
-			*OutAngularVelocityY.GetDest() = VValue.Y;
-			*OutAngularVelocityZ.GetDest() = VValue.Z;
-
-			FValue = InstData->IncomingMass1Array[ParticleID];
-			*OutMass.GetDest() = FValue;
-
-			FValue = InstData->IncomingTimeArray[ParticleID];
-			*OutTime.GetDest() = FValue;
-		}
-
-		ParticleIDParam.Advance();
-		OutLocationX.Advance();
-		OutLocationY.Advance();
-		OutLocationZ.Advance();
-		OutVelocityX.Advance();
-		OutVelocityY.Advance();
-		OutVelocityZ.Advance();
-		OutAngularVelocityX.Advance();
-		OutAngularVelocityY.Advance();
-		OutAngularVelocityZ.Advance();
-		OutMass.Advance();
-		OutTime.Advance();
+		OutLocation.SetAndAdvance(InstData->IncomingLocationArray.IsValidIndex(ParticleID) ? (FVector3f)InstData->IncomingLocationArray[ParticleID] : FVector3f::Zero());
+		OutVelocity.SetAndAdvance(InstData->IncomingVelocity1Array.IsValidIndex(ParticleID) ? (FVector3f)InstData->IncomingVelocity1Array[ParticleID] : FVector3f::Zero());
+		OutAngularVelocity.SetAndAdvance(InstData->IncomingAngularVelocity1Array.IsValidIndex(ParticleID) ? (FVector3f)InstData->IncomingAngularVelocity1Array[ParticleID] : FVector3f::Zero());
+		OutMass.SetAndAdvance(InstData->IncomingMass1Array.IsValidIndex(ParticleID) ? InstData->IncomingMass1Array[ParticleID] : 0.0f);
+		OutTime.SetAndAdvance(InstData->IncomingTimeArray.IsValidIndex(ParticleID) ? InstData->IncomingTimeArray[ParticleID] : 0.0f);
 	}
 }
 
@@ -4002,7 +3652,7 @@ static void SetBuffer(FRHICommandList& CmdList,
 	FDynamicReadBuffer& Buffer,
 	const TArray<T>& Array,
 	const EPixelFormat PixelFormat,
-	FString BufferName)
+	const TCHAR* BufferName)
 {
 	checkf(PixelFormat == PF_A32B32G32R32F ||
 		   PixelFormat == PF_R32_FLOAT || 
@@ -4023,7 +3673,7 @@ static void SetBuffer(FRHICommandList& CmdList,
 	if (Buffer.NumBytes == 0 || Buffer.NumBytes < Array.Num() * SizePerElement)
 	{
 		Buffer.Release();
-		Buffer.Initialize(*BufferName, SizePerElement, Array.Num(), PixelFormat, BUF_Dynamic);
+		Buffer.Initialize(BufferName, SizePerElement, Array.Num(), PixelFormat, BUF_Dynamic);
 	}
 
 	Buffer.Lock();
@@ -4087,12 +3737,22 @@ void UNiagaraDataInterfaceChaosDestruction::ProvidePerInstanceDataForRenderThrea
 
 	if (InstanceData->PositionArray.Num() > 0)
 	{
-		DataToPass->PositionArray = new TArray<FVector>(InstanceData->PositionArray);
+		DataToPass->PositionArray = new TArray<FVector3f>();
+		DataToPass->PositionArray->AddUninitialized(InstanceData->PositionArray.Num());
+		for (int i = 0; i < InstanceData->PositionArray.Num(); ++i)
+		{
+			(*DataToPass->PositionArray)[i] = (FVector3f)InstanceData->PositionArray[i];
+		}
 	}
 
 	if (InstanceData->VelocityArray.Num() > 0)
 	{
-		DataToPass->VelocityArray = new TArray<FVector>(InstanceData->VelocityArray);
+		DataToPass->VelocityArray = new TArray<FVector3f>();
+		DataToPass->VelocityArray->AddUninitialized(InstanceData->VelocityArray.Num());
+		for (int i = 0; i < InstanceData->VelocityArray.Num(); ++i)
+		{
+			(*DataToPass->VelocityArray)[i] = (FVector3f)InstanceData->VelocityArray[i];
+		}
 	}
 
 	if (InstanceData->ExtentMinArray.Num() > 0)
@@ -4240,14 +3900,14 @@ void FNiagaraDataInterfaceProxyChaosDestruction::ConsumePerInstanceDataFromGameT
 
 		if (InstanceData->PositionArray)
 		{
-			Data.PositionArray = TArray<FVector>(MoveTemp(*InstanceData->PositionArray));		
+			Data.PositionArray = TArray<FVector3f>(MoveTemp(*InstanceData->PositionArray));		
 			//LoadGPUBufferFromArray(Data.GPUPositionBuffer, InstanceData->PositionArray, EPixelFormat::PF_A32B32G32R32F, FString(TEXT("PositionBuffer")));
 			delete InstanceData->PositionArray;
 		}
 
 		if (InstanceData->VelocityArray)
 		{
-			Data.VelocityArray = TArray<FVector>(MoveTemp(*InstanceData->VelocityArray));
+			Data.VelocityArray = TArray<FVector3f>(MoveTemp(*InstanceData->VelocityArray));
 			//LoadGPUBufferFromArray(Data.GPUVelocityBuffer, InstanceData->VelocityArray, EPixelFormat::PF_A32B32G32R32F, FString(TEXT("VelocityBuffer")));
 			delete InstanceData->VelocityArray;
 		}
@@ -4472,33 +4132,33 @@ public:
 			}
 			if(InstanceData->PositionArray.Num() > 0)
 			{
-				SetBuffer(RHICmdList, PositionBuffer,    ComputeShaderRHI, InstanceData->GPUPositionBuffer,    InstanceData->PositionArray,    EPixelFormat::PF_A32B32G32R32F, FString(TEXT("PositionBuffer")));
-				SetBuffer(RHICmdList, VelocityBuffer,    ComputeShaderRHI, InstanceData->GPUVelocityBuffer,    InstanceData->VelocityArray,    EPixelFormat::PF_A32B32G32R32F, FString(TEXT("VelocityBuffer")));
-				SetBuffer(RHICmdList, ExtentMinBuffer,   ComputeShaderRHI, InstanceData->GPUExtentMinBuffer,   InstanceData->ExtentMinArray,   EPixelFormat::PF_R32_FLOAT,     FString(TEXT("ExtentMinBuffer")));
-				SetBuffer(RHICmdList, ExtentMaxBuffer,   ComputeShaderRHI, InstanceData->GPUExtentMaxBuffer,   InstanceData->ExtentMaxArray,   EPixelFormat::PF_R32_FLOAT,     FString(TEXT("ExtentMaxBuffer")));
-				SetBuffer(RHICmdList, VolumeBuffer,      ComputeShaderRHI, InstanceData->GPUVolumeBuffer,      InstanceData->VolumeArray,      EPixelFormat::PF_R32_FLOAT,     FString(TEXT("VolumeBuffer")));
-				SetBuffer(RHICmdList, SolverIDBuffer,    ComputeShaderRHI, InstanceData->GPUSolverIDBuffer,    InstanceData->SolverIDArray,    EPixelFormat::PF_R32_SINT,      FString(TEXT("SolverIDBuffer")));
-				SetBuffer(RHICmdList, DensityBuffer,     ComputeShaderRHI, InstanceData->GPUDensityBuffer,     InstanceData->DensityArray,     EPixelFormat::PF_R32_FLOAT,     FString(TEXT("DensityBuffer")));
-				SetBuffer(RHICmdList, FrictionBuffer,    ComputeShaderRHI, InstanceData->GPUFrictionBuffer,    InstanceData->FrictionArray,    EPixelFormat::PF_R32_FLOAT,     FString(TEXT("FrictionBuffer")));
-				SetBuffer(RHICmdList, RestitutionBuffer, ComputeShaderRHI, InstanceData->GPURestitutionBuffer, InstanceData->RestitutionArray, EPixelFormat::PF_R32_FLOAT,     FString(TEXT("RestitutionBuffer")));
-				SetBuffer(RHICmdList, SurfaceTypeBuffer, ComputeShaderRHI, InstanceData->GPUSurfaceTypeBuffer, InstanceData->SurfaceTypeArray, EPixelFormat::PF_R32_SINT,      FString(TEXT("SurfaceTypeBuffer")));
-				SetBuffer(RHICmdList, ColorBuffer,       ComputeShaderRHI, InstanceData->GPUColorBuffer,       InstanceData->ColorArray,       EPixelFormat::PF_A32B32G32R32F, FString(TEXT("ColorBuffer")));
-
-				SetBuffer(RHICmdList, IncomingLocationBuffer,           ComputeShaderRHI, InstanceData->GPUIncomingLocationBuffer,           InstanceData->IncomingLocationArray,           EPixelFormat::PF_A32B32G32R32F, FString(TEXT("IncomingLocationBuffer")));
-				SetBuffer(RHICmdList, IncomingAccumulatedImpulseBuffer, ComputeShaderRHI, InstanceData->GPUIncomingAccumulatedImpulseBuffer, InstanceData->IncomingAccumulatedImpulseArray, EPixelFormat::PF_A32B32G32R32F, FString(TEXT("IncomingAccumulatedImpulseBuffer")));
-				SetBuffer(RHICmdList, IncomingNormalBuffer,             ComputeShaderRHI, InstanceData->GPUIncomingNormalBuffer,             InstanceData->IncomingNormalArray,             EPixelFormat::PF_A32B32G32R32F, FString(TEXT("IncomingNormalBuffer")));
-				SetBuffer(RHICmdList, IncomingVelocity1Buffer,          ComputeShaderRHI, InstanceData->GPUIncomingVelocity1Buffer,          InstanceData->IncomingVelocity1Array,          EPixelFormat::PF_A32B32G32R32F, FString(TEXT("IncomingVelocity1Buffer")));
-				SetBuffer(RHICmdList, IncomingVelocity2Buffer,          ComputeShaderRHI, InstanceData->GPUIncomingVelocity2Buffer,          InstanceData->IncomingVelocity2Array,          EPixelFormat::PF_A32B32G32R32F, FString(TEXT("IncomingVelocity2Buffer")));
-				SetBuffer(RHICmdList, IncomingAngularVelocity1Buffer,   ComputeShaderRHI, InstanceData->GPUIncomingAngularVelocity1Buffer,   InstanceData->IncomingAngularVelocity1Array,   EPixelFormat::PF_A32B32G32R32F, FString(TEXT("IncomingAngularVelocity1Buffer")));
-				SetBuffer(RHICmdList, IncomingAngularVelocity2Buffer,   ComputeShaderRHI, InstanceData->GPUIncomingAngularVelocity2Buffer,   InstanceData->IncomingAngularVelocity2Array,   EPixelFormat::PF_A32B32G32R32F, FString(TEXT("IncomingAngularVelocity2Buffer")));
-				SetBuffer(RHICmdList, IncomingMass1Buffer,              ComputeShaderRHI, InstanceData->GPUIncomingMass1Buffer,              InstanceData->IncomingMass1Array,              EPixelFormat::PF_R32_FLOAT,     FString(TEXT("IncomingMass1Buffer")));
-				SetBuffer(RHICmdList, IncomingMass2Buffer,              ComputeShaderRHI, InstanceData->GPUIncomingMass2Buffer,              InstanceData->IncomingMass2Array,              EPixelFormat::PF_R32_FLOAT,     FString(TEXT("IncomingMass2Buffer")));
-				SetBuffer(RHICmdList, IncomingTimeBuffer,               ComputeShaderRHI, InstanceData->GPUIncomingTimeBuffer,               InstanceData->IncomingTimeArray,               EPixelFormat::PF_R32_FLOAT,     FString(TEXT("IncomingTimeBuffer")));
+				SetBuffer(RHICmdList, PositionBuffer,    ComputeShaderRHI, InstanceData->GPUPositionBuffer,    InstanceData->PositionArray,    EPixelFormat::PF_A32B32G32R32F, TEXT("PositionBuffer"));
+				SetBuffer(RHICmdList, VelocityBuffer,    ComputeShaderRHI, InstanceData->GPUVelocityBuffer,    InstanceData->VelocityArray,    EPixelFormat::PF_A32B32G32R32F, TEXT("VelocityBuffer"));
+				SetBuffer(RHICmdList, ExtentMinBuffer,   ComputeShaderRHI, InstanceData->GPUExtentMinBuffer,   InstanceData->ExtentMinArray,   EPixelFormat::PF_R32_FLOAT,     TEXT("ExtentMinBuffer"));
+				SetBuffer(RHICmdList, ExtentMaxBuffer,   ComputeShaderRHI, InstanceData->GPUExtentMaxBuffer,   InstanceData->ExtentMaxArray,   EPixelFormat::PF_R32_FLOAT,     TEXT("ExtentMaxBuffer"));
+				SetBuffer(RHICmdList, VolumeBuffer,      ComputeShaderRHI, InstanceData->GPUVolumeBuffer,      InstanceData->VolumeArray,      EPixelFormat::PF_R32_FLOAT,     TEXT("VolumeBuffer"));
+				SetBuffer(RHICmdList, SolverIDBuffer,    ComputeShaderRHI, InstanceData->GPUSolverIDBuffer,    InstanceData->SolverIDArray,    EPixelFormat::PF_R32_SINT,      TEXT("SolverIDBuffer"));
+				SetBuffer(RHICmdList, DensityBuffer,     ComputeShaderRHI, InstanceData->GPUDensityBuffer,     InstanceData->DensityArray,     EPixelFormat::PF_R32_FLOAT,     TEXT("DensityBuffer"));
+				SetBuffer(RHICmdList, FrictionBuffer,    ComputeShaderRHI, InstanceData->GPUFrictionBuffer,    InstanceData->FrictionArray,    EPixelFormat::PF_R32_FLOAT,     TEXT("FrictionBuffer"));
+				SetBuffer(RHICmdList, RestitutionBuffer, ComputeShaderRHI, InstanceData->GPURestitutionBuffer, InstanceData->RestitutionArray, EPixelFormat::PF_R32_FLOAT,     TEXT("RestitutionBuffer"));
+				SetBuffer(RHICmdList, SurfaceTypeBuffer, ComputeShaderRHI, InstanceData->GPUSurfaceTypeBuffer, InstanceData->SurfaceTypeArray, EPixelFormat::PF_R32_SINT,      TEXT("SurfaceTypeBuffer"));
+				SetBuffer(RHICmdList, ColorBuffer,       ComputeShaderRHI, InstanceData->GPUColorBuffer,       InstanceData->ColorArray,       EPixelFormat::PF_A32B32G32R32F, TEXT("ColorBuffer"));
+
+				SetBuffer(RHICmdList, IncomingLocationBuffer,           ComputeShaderRHI, InstanceData->GPUIncomingLocationBuffer,           InstanceData->IncomingLocationArray,           EPixelFormat::PF_A32B32G32R32F, TEXT("IncomingLocationBuffer"));
+				SetBuffer(RHICmdList, IncomingAccumulatedImpulseBuffer, ComputeShaderRHI, InstanceData->GPUIncomingAccumulatedImpulseBuffer, InstanceData->IncomingAccumulatedImpulseArray, EPixelFormat::PF_A32B32G32R32F, TEXT("IncomingAccumulatedImpulseBuffer"));
+				SetBuffer(RHICmdList, IncomingNormalBuffer,             ComputeShaderRHI, InstanceData->GPUIncomingNormalBuffer,             InstanceData->IncomingNormalArray,             EPixelFormat::PF_A32B32G32R32F, TEXT("IncomingNormalBuffer"));
+				SetBuffer(RHICmdList, IncomingVelocity1Buffer,          ComputeShaderRHI, InstanceData->GPUIncomingVelocity1Buffer,          InstanceData->IncomingVelocity1Array,          EPixelFormat::PF_A32B32G32R32F, TEXT("IncomingVelocity1Buffer"));
+				SetBuffer(RHICmdList, IncomingVelocity2Buffer,          ComputeShaderRHI, InstanceData->GPUIncomingVelocity2Buffer,          InstanceData->IncomingVelocity2Array,          EPixelFormat::PF_A32B32G32R32F, TEXT("IncomingVelocity2Buffer"));
+				SetBuffer(RHICmdList, IncomingAngularVelocity1Buffer,   ComputeShaderRHI, InstanceData->GPUIncomingAngularVelocity1Buffer,   InstanceData->IncomingAngularVelocity1Array,   EPixelFormat::PF_A32B32G32R32F, TEXT("IncomingAngularVelocity1Buffer"));
+				SetBuffer(RHICmdList, IncomingAngularVelocity2Buffer,   ComputeShaderRHI, InstanceData->GPUIncomingAngularVelocity2Buffer,   InstanceData->IncomingAngularVelocity2Array,   EPixelFormat::PF_A32B32G32R32F, TEXT("IncomingAngularVelocity2Buffer"));
+				SetBuffer(RHICmdList, IncomingMass1Buffer,              ComputeShaderRHI, InstanceData->GPUIncomingMass1Buffer,              InstanceData->IncomingMass1Array,              EPixelFormat::PF_R32_FLOAT,     TEXT("IncomingMass1Buffer"));
+				SetBuffer(RHICmdList, IncomingMass2Buffer,              ComputeShaderRHI, InstanceData->GPUIncomingMass2Buffer,              InstanceData->IncomingMass2Array,              EPixelFormat::PF_R32_FLOAT,     TEXT("IncomingMass2Buffer"));
+				SetBuffer(RHICmdList, IncomingTimeBuffer,               ComputeShaderRHI, InstanceData->GPUIncomingTimeBuffer,               InstanceData->IncomingTimeArray,               EPixelFormat::PF_R32_FLOAT,     TEXT("IncomingTimeBuffer"));
 				
-				SetBuffer(RHICmdList, TransformTranslationBuffer,       ComputeShaderRHI, InstanceData->GPUTransformTranslationBuffer,       InstanceData->TransformTranslationArray,       EPixelFormat::PF_A32B32G32R32F, FString(TEXT("TransformTranslationBuffer")));
-				SetBuffer(RHICmdList, TransformRotationBuffer,          ComputeShaderRHI, InstanceData->GPUTransformRotationBuffer,          InstanceData->TransformRotationArray,          EPixelFormat::PF_A32B32G32R32F, FString(TEXT("TransformRotationBuffer")));
-				SetBuffer(RHICmdList, TransformScaleBuffer,             ComputeShaderRHI, InstanceData->GPUTransformScaleBuffer,             InstanceData->TransformScaleArray,             EPixelFormat::PF_A32B32G32R32F, FString(TEXT("TransformScaleBuffer")));
-				SetBuffer(RHICmdList, BoundsBuffer,                     ComputeShaderRHI, InstanceData->GPUBoundsBuffer,                     InstanceData->BoundsArray,                     EPixelFormat::PF_A32B32G32R32F, FString(TEXT("BoundsBuffer")));
+				SetBuffer(RHICmdList, TransformTranslationBuffer,       ComputeShaderRHI, InstanceData->GPUTransformTranslationBuffer,       InstanceData->TransformTranslationArray,       EPixelFormat::PF_A32B32G32R32F, TEXT("TransformTranslationBuffer"));
+				SetBuffer(RHICmdList, TransformRotationBuffer,          ComputeShaderRHI, InstanceData->GPUTransformRotationBuffer,          InstanceData->TransformRotationArray,          EPixelFormat::PF_A32B32G32R32F, TEXT("TransformRotationBuffer"));
+				SetBuffer(RHICmdList, TransformScaleBuffer,             ComputeShaderRHI, InstanceData->GPUTransformScaleBuffer,             InstanceData->TransformScaleArray,             EPixelFormat::PF_A32B32G32R32F, TEXT("TransformScaleBuffer"));
+				SetBuffer(RHICmdList, BoundsBuffer,                     ComputeShaderRHI, InstanceData->GPUBoundsBuffer,                     InstanceData->BoundsArray,                     EPixelFormat::PF_A32B32G32R32F, TEXT("BoundsBuffer"));
 
 				SetShaderValue(RHICmdList, ComputeShaderRHI, LastSpawnedPointID, InstanceData->LastSpawnedPointID);
 				SetShaderValue(RHICmdList, ComputeShaderRHI, SolverTime, InstanceData->SolverTime);
