- Set version to latest after upgrade

#rnx

diff --git a/Engine/Plugins/Experimental/ChaosNiagara/Source/ChaosNiagara/Private/NiagaraDataInterfacePhysicsField.cpp b/Engine/Plugins/Experimental/ChaosNiagara/Source/ChaosNiagara/Private/NiagaraDataInterfacePhysicsField.cpp
--- a/Engine/Plugins/Experimental/ChaosNiagara/Source/ChaosNiagara/Private/NiagaraDataInterfacePhysicsField.cpp
+++ b/Engine/Plugins/Experimental/ChaosNiagara/Source/ChaosNiagara/Private/NiagaraDataInterfacePhysicsField.cpp
@@ -801,6 +801,7 @@ bool UNiagaraDataInterfacePhysicsField::UpgradeFunctionCall(FNiagaraFunctionSign
 			bChanged = true;
 		}
 	}
+	FunctionSignature.FunctionVersion = FNiagaraPhysicsFieldDIFunctionVersion::LatestVersion;
 
 	return bChanged;
 }
diff --git a/Engine/Plugins/Experimental/Water/Source/Runtime/Private/NiagaraDataInterfaceWater.cpp b/Engine/Plugins/Experimental/Water/Source/Runtime/Private/NiagaraDataInterfaceWater.cpp
--- a/Engine/Plugins/Experimental/Water/Source/Runtime/Private/NiagaraDataInterfaceWater.cpp
+++ b/Engine/Plugins/Experimental/Water/Source/Runtime/Private/NiagaraDataInterfaceWater.cpp
@@ -179,6 +179,7 @@ bool UNiagaraDataInterfaceWater::UpgradeFunctionCall(FNiagaraFunctionSignature&
 			bChanged = true;
 		}
 	}
+	FunctionSignature.FunctionVersion = FNiagaraWaterDIFunctionVersion::LatestVersion;
 
 	return bChanged;
 }
diff --git a/Engine/Plugins/FX/Niagara/Source/Niagara/Private/NiagaraDataInterfaceRigidMeshCollisionQuery.cpp b/Engine/Plugins/FX/Niagara/Source/Niagara/Private/NiagaraDataInterfaceRigidMeshCollisionQuery.cpp
--- a/Engine/Plugins/FX/Niagara/Source/Niagara/Private/NiagaraDataInterfaceRigidMeshCollisionQuery.cpp
+++ b/Engine/Plugins/FX/Niagara/Source/Niagara/Private/NiagaraDataInterfaceRigidMeshCollisionQuery.cpp
@@ -1233,6 +1233,7 @@ bool UNiagaraDataInterfaceRigidMeshCollisionQuery::UpgradeFunctionCall(FNiagaraF
 			bChanged = true;
 		}
 	}
+	FunctionSignature.FunctionVersion = FNiagaraRigidMeshCollisionDIFunctionVersion::LatestVersion;
 
 	return bChanged;
 }
