Fixed a bug where physics field data interface fails to compile for gpu particles

diff --git a/Engine/Plugins/Experimental/ChaosNiagara/Shaders/NiagaraDataInterfacePhysicsField.ush b/Engine/Plugins/Experimental/ChaosNiagara/Shaders/NiagaraDataInterfacePhysicsField.ush
--- a/Engine/Plugins/Experimental/ChaosNiagara/Shaders/NiagaraDataInterfacePhysicsField.ush
+++ b/Engine/Plugins/Experimental/ChaosNiagara/Shaders/NiagaraDataInterfacePhysicsField.ush
@@ -47,7 +47,7 @@ DIContext.ClipmapExponent = ClipmapExponent_##NAME;\
 DIContext.ClipmapCount = ClipmapCount_##NAME;\
 DIContext.TargetCount = TargetCount_##NAME;\
 DIContext.FieldTargets = FieldTargets_##NAME;\
-DIContext.SystemLWCTile = SystemLWCTile##NAME;\
+DIContext.SystemLWCTile = SystemLWCTile_##NAME;\
 
 float3 DIPhysicsField_SamplePhysicsVectorField(in FDIPhysicsFieldContext DIContext, in float3 WorldPosition, in int VectorTarget)
 {
