Chaos : NiagaraDataInterface
- Update to support trailing data

diff --git a/Engine/Plugins/Experimental/ChaosNiagara/Source/ChaosNiagara/Classes/NiagaraDataInterfaceChaosDestruction.h b/Engine/Plugins/Experimental/ChaosNiagara/Source/ChaosNiagara/Classes/NiagaraDataInterfaceChaosDestruction.h
--- a/Engine/Plugins/Experimental/ChaosNiagara/Source/ChaosNiagara/Classes/NiagaraDataInterfaceChaosDestruction.h
+++ b/Engine/Plugins/Experimental/ChaosNiagara/Source/ChaosNiagara/Classes/NiagaraDataInterfaceChaosDestruction.h
@@ -443,6 +443,10 @@ public:
 	UPROPERTY(EditAnywhere, Category = "Spawn Settings - Thresholds to Spawn", meta = (DisplayName = "Min/Max LocationZ To Spawn Particles", UIMin = 0.0))
 	FVector2D LocationZToSpawnMinMax;
 
+	/* Min Linear Speed to generate trailing particles */
+	UPROPERTY(EditAnywhere, Category = "Spawn Settings - Trailing - Thresholds to Spawn", meta = (DisplayName = "Min Speed To Spawn Trailing Particles", UIMin = 0.0))
+	float TrailMinSpeedToSpawn;
+
 	/**
 	* Sorting method to sort the collision data
 	*/
@@ -474,19 +478,19 @@ public:
 	int32 MaxDataPerCell;
 
 	/* Materials Filter */
-	UPROPERTY(EditAnywhere, Category = "Breaking Data Settings", meta = (DisplayName = "Apply Materials Filter"))
+	UPROPERTY(EditAnywhere, Category = "Data Settings", meta = (DisplayName = "Apply Materials Filter"))
 	bool bApplyMaterialsFilter;
 
 	/* TODO: Explanatory comment */
-	UPROPERTY(EditAnywhere, Category = "Breaking Data Settings", meta = (DisplayName = "Collision/Breaking Filtered Materials", EditCondition = bApplyMaterialsFilter))
+	UPROPERTY(EditAnywhere, Category = "Data Settings", meta = (DisplayName = "Filtered Materials", EditCondition = bApplyMaterialsFilter))
 	TSet<TObjectPtr<UPhysicalMaterial>> ChaosBreakingMaterialSet;
 
 	/* TODO: Explanatory comment */
-	UPROPERTY(EditAnywhere, Category = "Breaking Data Settings", meta = (DisplayName = "Get External Breaking Mesh and Physical Data"))
+	UPROPERTY(EditAnywhere, Category = "Data Settings", meta = (DisplayName = "Get External Breaking Mesh and Physical Data"))
 	bool bGetExternalBreakingData;
 
 	/* TODO: Explanatory comment */
-	UPROPERTY(VisibleAnywhere, Category = "Trailing Data Settings", meta = (DisplayName = "Get External Trailing Mesh and Physical Data"))
+	UPROPERTY(VisibleAnywhere, Category = "Data Settings", meta = (DisplayName = "Get External Trailing Mesh and Physical Data"))
 	bool bGetExternalTrailingData;
 
 	/* Random displacement value for the particle spawn position */
@@ -758,6 +762,9 @@ protected:
 	TArray<Chaos::FCollidingDataExt> CollisionEvents;
 	TArray<Chaos::FBreakingDataExt> BreakingEvents;
 	TArray<Chaos::FTrailingDataExt> TrailingEvents;
+
+	// Saving GeometryCollectionComponents from breaking for trailing
+	TSet<UGeometryCollectionComponent*> GeometryCollectionComponentsFromBreaking;
 };
 
 struct FNiagaraDataInterfaceProxyChaosDestruction : public FNiagaraDataInterfaceProxy
diff --git a/Engine/Plugins/Experimental/ChaosNiagara/Source/ChaosNiagara/Private/NiagaraDataInterfaceChaosDestruction.cpp b/Engine/Plugins/Experimental/ChaosNiagara/Source/ChaosNiagara/Private/NiagaraDataInterfaceChaosDestruction.cpp
--- a/Engine/Plugins/Experimental/ChaosNiagara/Source/ChaosNiagara/Private/NiagaraDataInterfaceChaosDestruction.cpp
+++ b/Engine/Plugins/Experimental/ChaosNiagara/Source/ChaosNiagara/Private/NiagaraDataInterfaceChaosDestruction.cpp
@@ -92,6 +92,7 @@ UNiagaraDataInterfaceChaosDestruction::UNiagaraDataInterfaceChaosDestruction(FOb
 	, LocationYToSpawnMinMax(FVector2D(0.f, 0.f))
 	, LocationZToSpawn(ELocationZToSpawnEnum::ChaosNiagara_LocationZToSpawn_None)
 	, LocationZToSpawnMinMax(FVector2D(0.f, 0.f))
+	, TrailMinSpeedToSpawn(50.f)
 	, DataSortingType(EDataSortTypeEnum::ChaosNiagara_DataSortType_NoSorting)
 	, bGetExternalCollisionData(false)
 	, DoSpatialHash(false)
@@ -768,59 +769,45 @@ void UNiagaraDataInterfaceChaosDestruction::HandleCollisionEvents(const Chaos::F
 
 		if (GeometryCollectionComponent)
 		{
-			Material = GeometryCollectionComponent->GetMaterial(MaterialID);
-		}
-#endif
-		if (Material)
-		{
-			PhysicalMaterial = Material->GetPhysicalMaterial();
+			PhysicalMaterial = GeometryCollectionComponent->GetPhysicalMaterial();
 			ensure(PhysicalMaterial);
 			if (PhysicalMaterial)
 			{
 				CopyData.SurfaceType1 = PhysicalMaterial->SurfaceType;
+				CopyData.PhysicalMaterialName1 = PhysicalMaterial->GetFName();
+			}
+			else
+			{
+				CopyData.PhysicalMaterialName1 = FName();
 			}
 		}
 
-		if (PhysicalMaterial)
-		{
-			CopyData.PhysicalMaterialName1 = PhysicalMaterial->GetFName();
-		}
-		else
-		{
-			CopyData.PhysicalMaterialName1 = FName();
-		}
-
-#if INCLUDE_CHAOS
 		for (auto& Solver : Solvers)
 		{
-			GeometryCollectionComponent = Solver.PhysScene->GetOwningComponent<UGeometryCollectionComponent>(DataIn.Proxy2);
-			if (GeometryCollectionComponent)
-				break;
+			if (Solver.PhysScene)
+			{
+				GeometryCollectionComponent = Solver.PhysScene->GetOwningComponent<UGeometryCollectionComponent>(DataIn.Proxy2);
+				if (GeometryCollectionComponent)
+					break;
+			}
 		}
+
 		if (GeometryCollectionComponent)
 		{
-			Material = GeometryCollectionComponent->GetMaterial(MaterialID);
-		}
-#endif
-		if (Material)
-		{
-			PhysicalMaterial = Material->GetPhysicalMaterial();
+			PhysicalMaterial = GeometryCollectionComponent->GetPhysicalMaterial();
 			ensure(PhysicalMaterial);
 			if (PhysicalMaterial)
 			{
 				CopyData.SurfaceType2 = PhysicalMaterial->SurfaceType;
+				CopyData.PhysicalMaterialName2 = PhysicalMaterial->GetFName();
+			}
+			else
+			{
+				CopyData.PhysicalMaterialName2 = FName();
 			}
 		}
+#endif
 
-		if (PhysicalMaterial)
-		{
-			CopyData.PhysicalMaterialName2 = PhysicalMaterial->GetFName();
-		}
-		else
-		{
-			CopyData.PhysicalMaterialName2 = FName();
-		}
-		
 		CopyData.BoundingboxVolume = 1000000.f;
 		CopyData.BoundingboxExtentMin = 100.f;
 		CopyData.BoundingboxExtentMax = 100.f;
@@ -1367,29 +1354,25 @@ void UNiagaraDataInterfaceChaosDestruction::HandleBreakingEvents(const Chaos::FB
 						break;
 				}
 			}
+
 			if (GeometryCollectionComponent)
 			{
-				Material = GeometryCollectionComponent->GetMaterial(MaterialID);
-			}
-#endif
-			if (Material)
-			{
-				PhysicalMaterial = Material->GetPhysicalMaterial();
+				PhysicalMaterial = GeometryCollectionComponent->GetPhysicalMaterial();
 				ensure(PhysicalMaterial);
 				if (PhysicalMaterial)
 				{
 					CopyData.SurfaceType = PhysicalMaterial->SurfaceType;
+					CopyData.PhysicalMaterialName = PhysicalMaterial->GetFName();
+				}
+				else
+				{
+					CopyData.PhysicalMaterialName = FName();
 				}
-			}
 
-			if (PhysicalMaterial)
-			{
-				CopyData.PhysicalMaterialName = PhysicalMaterial->GetFName();
-			}
-			else
-			{
-				CopyData.PhysicalMaterialName = FName();
+				// Save GeometryCollectionComponent for trailing
+				GeometryCollectionComponentsFromBreaking.Add(GeometryCollectionComponent);
 			}
+#endif
 		}
 		else
 		{
@@ -1887,37 +1870,93 @@ bool UNiagaraDataInterfaceChaosDestruction::BreakingCallback(FNDIChaosDestructio
 void UNiagaraDataInterfaceChaosDestruction::HandleTrailingEvents(const Chaos::FTrailingEventData& Event)
 {
 	ensure(IsInGameThread());
-	Chaos::FTrailingDataArray const& TrailingDataIn = Event.TrailingData.AllTrailingsArray;
 
 	TrailingEvents.Reset();
 
-	// Copy data from Event
-	TrailingEvents.AddUninitialized(Event.TrailingData.AllTrailingsArray.Num());
+	float TrailMinSpeedToSpawnSquared = TrailMinSpeedToSpawn * TrailMinSpeedToSpawn;
+	UPhysicalMaterial* PhysicalMaterial = nullptr;
 
-	int32 Idx = 0;
-	for (Chaos::FTrailingData const& DataIn : TrailingDataIn)
+	auto IsMaterialInFilter = [&](const FName& InMaterialName) 
 	{
-		auto& CopyData = TrailingEvents[Idx];
-		CopyData = DataIn;
+		if (!InMaterialName.IsValid())
+		{
+			return false;
+		}
 
-		// #GM: Disable this for now for perf
-		/*
-		GetMeshExtData(SolverData,
-			AllTrailingsArray[Idx].ParticleIndexMesh == INDEX_NONE ? AllTrailingsArray[Idx].ParticleIndex : AllTrailingsArray[Idx].ParticleIndexMesh,
-			PhysicsProxyReverseMappingArray,
-			ParticleIndexReverseMappingArray,
-			AllTrailingsArray[Idx].BoundingboxVolume,
-			AllTrailingsArray[Idx].BoundingboxExtentMin,
-			AllTrailingsArray[Idx].BoundingboxExtentMax,
-			AllTrailingsArray[Idx].SurfaceType);
-		*/
+		for (const UPhysicalMaterial* Material : ChaosBreakingMaterialSet)
+		{
+			if (!Material)
+			{
+				continue;
+			}
 
-		CopyData.BoundingboxVolume = 1000000.f;
-		CopyData.BoundingboxExtentMin = 100.f;
-		CopyData.BoundingboxExtentMax = 100.f;
-		CopyData.SurfaceType = 0;
+			if (Material->GetFName() == InMaterialName)
+			{
+				return true;
+			}
+		}
 
-		Idx++;
+		return false;
+	};
+
+	if (GeometryCollectionComponentsFromBreaking.Num() > 0)
+	{
+		for (auto& GeometryCollectionComponent : GeometryCollectionComponentsFromBreaking)
+		{
+			if (GeometryCollectionComponent && GeometryCollectionComponent->GetNotifyTrailing())
+			{
+				PhysicalMaterial = GeometryCollectionComponent->GetPhysicalMaterial();
+				ensure(PhysicalMaterial);
+				if (PhysicalMaterial)
+				{
+					if (!(bApplyMaterialsFilter && IsMaterialInFilter(PhysicalMaterial->GetFName())))
+					{
+						continue;
+					}
+				}
+
+				const TArray<FMatrix>& GlobalMatrices = GeometryCollectionComponent->GetGlobalMatrices();
+				const FTransform ActorTransform = GeometryCollectionComponent->GetComponentToWorld();
+
+				const FGeometryDynamicCollection* DynamicCollection = GeometryCollectionComponent->GetDynamicCollection();
+
+				if (DynamicCollection)
+				{
+					const TManagedArray<FVector3f>* LinearVelocity = DynamicCollection->FindAttributeTyped<FVector3f>("LinearVelocity", FTransformCollection::TransformGroup);
+					const TManagedArray<FVector3f>* AngularVelocity = DynamicCollection->FindAttributeTyped<FVector3f>("AngularVelocity", FTransformCollection::TransformGroup);
+
+					if (!LinearVelocity || !AngularVelocity)
+					{
+						continue;
+					}
+
+					ensure(DynamicCollection->Active.Num() == (*LinearVelocity).Num());
+					ensure(DynamicCollection->Active.Num() == (*AngularVelocity).Num());
+
+					for (int Idx = 0; Idx < DynamicCollection->Active.Num(); ++Idx)
+					{
+						if (DynamicCollection->Active[Idx] && (*LinearVelocity)[Idx].SquaredLength() >= TrailMinSpeedToSpawnSquared)
+						{
+							Chaos::FTrailingDataExt TrailingData;
+
+							FTransform CurrTransform = FTransform(GlobalMatrices[Idx] * ActorTransform.ToMatrixWithScale());
+							TrailingData.Location = CurrTransform.GetTranslation();
+
+							TrailingData.Velocity = (*LinearVelocity)[Idx];
+							TrailingData.AngularVelocity = (*AngularVelocity)[Idx];
+
+							TrailingData.Mass = 1.f;
+							TrailingData.BoundingboxVolume = 1000000.f;
+							TrailingData.BoundingboxExtentMin = 100.f;
+							TrailingData.BoundingboxExtentMax = 100.f;
+							TrailingData.SurfaceType = 0;
+
+							TrailingEvents.Add(TrailingData);
+						}
+					}
+				}				
+			}
+		}
 	}
 }
 
diff --git a/Engine/Source/Runtime/Experimental/GeometryCollectionEngine/Private/GeometryCollection/GeometryCollectionComponent.cpp b/Engine/Source/Runtime/Experimental/GeometryCollectionEngine/Private/GeometryCollection/GeometryCollectionComponent.cpp
--- a/Engine/Source/Runtime/Experimental/GeometryCollectionEngine/Private/GeometryCollection/GeometryCollectionComponent.cpp
+++ b/Engine/Source/Runtime/Experimental/GeometryCollectionEngine/Private/GeometryCollection/GeometryCollectionComponent.cpp
@@ -244,6 +244,7 @@ UGeometryCollectionComponent::UGeometryCollectionComponent(const FObjectInitiali
 	, bNotifyBreaks(false)
 	, bNotifyCollisions(false)
 	, bNotifyRemovals(false)
+	, bStoreVelocities(false)
 	, bShowBoneColors(false)
 	, bEnableReplication(false)
 	, bEnableAbandonAfterLevel(false)
@@ -846,7 +847,7 @@ UPhysicalMaterial* UGeometryCollectionComponent::GetPhysicalMaterial() const
 	// Pull material from first mesh element to grab physical material. Prefer an override if one exists
 	UPhysicalMaterial* PhysMatToUse = BodyInstance.GetSimplePhysicalMaterial();
 
-	if(!PhysMatToUse)
+	if(!PhysMatToUse || PhysMatToUse->GetFName() == "DefaultPhysicalMaterial")
 	{
 		// No override, try render materials
 		const int32 NumMaterials = GetNumMaterials();
@@ -1022,6 +1023,19 @@ void UGeometryCollectionComponent::InitializeComponent()
 {
 	Super::InitializeComponent();
 
+	if (bStoreVelocities || bNotifyTrailing)
+	{
+		if (!DynamicCollection->FindAttributeTyped<FVector3f>("LinearVelocity", FTransformCollection::TransformGroup))
+		{
+			DynamicCollection->AddAttribute<FVector3f>("LinearVelocity", FTransformCollection::TransformGroup);
+		}
+
+		if (!DynamicCollection->FindAttributeTyped<FVector3f>("AngularVelocity", FTransformCollection::TransformGroup))
+		{
+			DynamicCollection->AddAttribute<FVector3f>("AngularVelocity", FTransformCollection::TransformGroup);
+		}
+	}
+
 	AActor* Owner = GetOwner();
 
 	if(!Owner)
@@ -1041,8 +1055,15 @@ void UGeometryCollectionComponent::InitializeComponent()
 			// As we're the authority we need to track velocities in the dynamic collection so we
 			// can send them over to the other clients to correctly set their state. Attach this now.
 			// The physics proxy will pick them up and populate them as needed
-			DynamicCollection->AddAttribute<FVector3f>("LinearVelocity", FTransformCollection::TransformGroup);
-			DynamicCollection->AddAttribute<FVector3f>("AngularVelocity", FTransformCollection::TransformGroup);
+			if (!DynamicCollection->FindAttributeTyped<FVector3f>("LinearVelocity", FTransformCollection::TransformGroup))
+			{
+				DynamicCollection->AddAttribute<FVector3f>("LinearVelocity", FTransformCollection::TransformGroup);
+			}
+
+			if (!DynamicCollection->FindAttributeTyped<FVector3f>("AngularVelocity", FTransformCollection::TransformGroup))
+			{
+				DynamicCollection->AddAttribute<FVector3f>("AngularVelocity", FTransformCollection::TransformGroup);
+			}
 
 			// We also need to track our control of particles if that control can be shared between server and client
 			if(bEnableAbandonAfterLevel)
diff --git a/Engine/Source/Runtime/Experimental/GeometryCollectionEngine/Public/GeometryCollection/GeometryCollectionComponent.h b/Engine/Source/Runtime/Experimental/GeometryCollectionEngine/Public/GeometryCollection/GeometryCollectionComponent.h
--- a/Engine/Source/Runtime/Experimental/GeometryCollectionEngine/Public/GeometryCollection/GeometryCollectionComponent.h
+++ b/Engine/Source/Runtime/Experimental/GeometryCollectionEngine/Public/GeometryCollection/GeometryCollectionComponent.h
@@ -586,6 +586,8 @@ public:
 	*/
 	/* ---------------------------------------------------------------------------------------- */
 	
+	bool GetNotifyTrailing() const { return bNotifyTrailing; }
+
 	void SetShowBoneColors(bool ShowBoneColorsIn) { bShowBoneColors = ShowBoneColorsIn; }
 	bool GetShowBoneColors() const { return bShowBoneColors; }
 	bool GetEnableBoneSelection() const { return bEnableBoneSelection; }
@@ -725,6 +727,10 @@ public:
 	UPROPERTY(EditAnywhere, BlueprintReadOnly, Category = "ChaosPhysics|General")
 	bool bNotifyRemovals;
 
+	/** If true, this component will save linear and angular velocities on its DynamicCollection. */
+	UPROPERTY(EditAnywhere, BlueprintReadOnly, Category = "ChaosPhysics|General")
+	bool bStoreVelocities;
+
 protected:
 	/** Display Bone Colors instead of assigned materials */
 	UPROPERTY(EditAnywhere, BlueprintReadWrite, Category = "ChaosPhysics|General")
