ChaosNiagara : NiagaraDataInterface
Bug fix for PhysicalMaterial filtering condition

diff --git a/Engine/Plugins/Experimental/ChaosNiagara/Source/ChaosNiagara/Private/NiagaraDataInterfaceChaosDestruction.cpp b/Engine/Plugins/Experimental/ChaosNiagara/Source/ChaosNiagara/Private/NiagaraDataInterfaceChaosDestruction.cpp
--- a/Engine/Plugins/Experimental/ChaosNiagara/Source/ChaosNiagara/Private/NiagaraDataInterfaceChaosDestruction.cpp
+++ b/Engine/Plugins/Experimental/ChaosNiagara/Source/ChaosNiagara/Private/NiagaraDataInterfaceChaosDestruction.cpp
@@ -880,8 +880,9 @@ void UNiagaraDataInterfaceChaosDestruction::FilterAllCollisions(TArray<Chaos::FC
 			float CollisionAccumulatedImpulseSquared = AllCollisionsArray[IdxCollision].AccumulatedImpulse.SizeSquared();
 			float CollisionSpeedSquared = AllCollisionsArray[IdxCollision].Velocity1.SizeSquared();
 
+			// @TODO(Chaos): Need to be optimized
 			if (/*(ParticleToProcess != nullptr && AllCollisionsArrayInOut[IdxCollision].Particle != ParticleToProcess) ||*/
-				!(bApplyMaterialsFilter && (IsMaterialInFilter(AllCollisionsArray[IdxCollision].PhysicalMaterialName1) || IsMaterialInFilter(AllCollisionsArray[IdxCollision].PhysicalMaterialName2))) ||
+				(bApplyMaterialsFilter && !IsMaterialInFilter(AllCollisionsArray[IdxCollision].PhysicalMaterialName1) && !IsMaterialInFilter(AllCollisionsArray[IdxCollision].PhysicalMaterialName2)) ||
 				(ImpulseToSpawnMinMax.X > 0.f && ImpulseToSpawnMinMax.Y < 0.f && CollisionAccumulatedImpulseSquared < MinImpulseToSpawnSquared) ||
 				(ImpulseToSpawnMinMax.X < 0.f && ImpulseToSpawnMinMax.Y > 0.f && CollisionAccumulatedImpulseSquared > MaxImpulseToSpawnSquared) ||
 				(ImpulseToSpawnMinMax.X > 0.f && ImpulseToSpawnMinMax.Y > 0.f && (CollisionAccumulatedImpulseSquared < MinImpulseToSpawnSquared || CollisionAccumulatedImpulseSquared > MaxImpulseToSpawnSquared)) ||
@@ -1446,7 +1447,7 @@ void UNiagaraDataInterfaceChaosDestruction::FilterAllBreakings(TArray<Chaos::FBr
 		{
 			float BreakingSpeedSquared = AllBreakingsArray[IdxBreaking].Velocity.SizeSquared();
 
-			if (!(bApplyMaterialsFilter && IsMaterialInFilter(AllBreakingsArray[IdxBreaking].PhysicalMaterialName)) ||
+			if ((bApplyMaterialsFilter && !IsMaterialInFilter(AllBreakingsArray[IdxBreaking].PhysicalMaterialName)) ||
 				(SpeedToSpawnMinMax.X > 0.f && SpeedToSpawnMinMax.Y < 0.f && BreakingSpeedSquared < MinSpeedToSpawnSquared) ||
 				(SpeedToSpawnMinMax.X < 0.f && SpeedToSpawnMinMax.Y > 0.f && BreakingSpeedSquared > MaxSpeedToSpawnSquared) ||
 				(SpeedToSpawnMinMax.X > 0.f && SpeedToSpawnMinMax.Y > 0.f && (BreakingSpeedSquared < MinSpeedToSpawnSquared || BreakingSpeedSquared > MaxSpeedToSpawnSquared)) ||
@@ -1909,7 +1910,7 @@ void UNiagaraDataInterfaceChaosDestruction::HandleTrailingEvents(const Chaos::FT
 				ensure(PhysicalMaterial);
 				if (PhysicalMaterial)
 				{
-					if (!(bApplyMaterialsFilter && IsMaterialInFilter(PhysicalMaterial->GetFName())))
+					if (bApplyMaterialsFilter && !IsMaterialInFilter(PhysicalMaterial->GetFName()))
 					{
 						continue;
 					}
