ChaosNiagara : NiagaraDataInterface
Update to trailing data transforms

diff --git a/Engine/Plugins/Experimental/ChaosNiagara/Source/ChaosNiagara/Private/NiagaraDataInterfaceChaosDestruction.cpp b/Engine/Plugins/Experimental/ChaosNiagara/Source/ChaosNiagara/Private/NiagaraDataInterfaceChaosDestruction.cpp
--- a/Engine/Plugins/Experimental/ChaosNiagara/Source/ChaosNiagara/Private/NiagaraDataInterfaceChaosDestruction.cpp
+++ b/Engine/Plugins/Experimental/ChaosNiagara/Source/ChaosNiagara/Private/NiagaraDataInterfaceChaosDestruction.cpp
@@ -1923,36 +1923,47 @@ void UNiagaraDataInterfaceChaosDestruction::HandleTrailingEvents(const Chaos::FT
 
 				if (DynamicCollection)
 				{
-					const TManagedArray<FVector3f>* LinearVelocity = DynamicCollection->FindAttributeTyped<FVector3f>("LinearVelocity", FTransformCollection::TransformGroup);
-					const TManagedArray<FVector3f>* AngularVelocity = DynamicCollection->FindAttributeTyped<FVector3f>("AngularVelocity", FTransformCollection::TransformGroup);
-
-					if (!LinearVelocity || !AngularVelocity)
+					if (const UGeometryCollection* RestCollection = GeometryCollectionComponent->GetRestCollection())
 					{
-						continue;
-					}
-
-					ensure(DynamicCollection->Active.Num() == (*LinearVelocity).Num());
-					ensure(DynamicCollection->Active.Num() == (*AngularVelocity).Num());
+						const TSharedPtr<FGeometryCollection, ESPMode::ThreadSafe> GeometryCollection = RestCollection->GetGeometryCollection();
 
-					for (int Idx = 0; Idx < DynamicCollection->Active.Num(); ++Idx)
-					{
-						if (DynamicCollection->Active[Idx] && (*LinearVelocity)[Idx].SquaredLength() >= TrailMinSpeedToSpawnSquared)
+						// Get the MassToLocal transforms
+						if (GeometryCollection->HasAttribute(TEXT("MassToLocal"), FTransformCollection::TransformGroup))
 						{
-							Chaos::FTrailingDataExt TrailingData;
+							const TManagedArray<FTransform>& CollectionMassToLocal = GeometryCollection->GetAttribute<FTransform>(TEXT("MassToLocal"), FTransformCollection::TransformGroup);
 
-							FTransform CurrTransform = FTransform(GlobalMatrices[Idx] * ActorTransform.ToMatrixWithScale());
-							TrailingData.Location = CurrTransform.GetTranslation();
+							const TManagedArray<FVector3f>* LinearVelocity = DynamicCollection->FindAttributeTyped<FVector3f>("LinearVelocity", FTransformCollection::TransformGroup);
+							const TManagedArray<FVector3f>* AngularVelocity = DynamicCollection->FindAttributeTyped<FVector3f>("AngularVelocity", FTransformCollection::TransformGroup);
 
-							TrailingData.Velocity = (*LinearVelocity)[Idx];
-							TrailingData.AngularVelocity = (*AngularVelocity)[Idx];
+							if (!LinearVelocity || !AngularVelocity)
+							{
+								continue;
+							}
 
-							TrailingData.Mass = 1.f;
-							TrailingData.BoundingboxVolume = 1000000.f;
-							TrailingData.BoundingboxExtentMin = 100.f;
-							TrailingData.BoundingboxExtentMax = 100.f;
-							TrailingData.SurfaceType = 0;
+							ensure(DynamicCollection->Active.Num() == (*LinearVelocity).Num());
+							ensure(DynamicCollection->Active.Num() == (*AngularVelocity).Num());
 
-							TrailingEvents.Add(TrailingData);
+							for (int Idx = 0; Idx < DynamicCollection->Active.Num(); ++Idx)
+							{
+								if (DynamicCollection->Active[Idx] && (*LinearVelocity)[Idx].SquaredLength() >= TrailMinSpeedToSpawnSquared)
+								{
+									Chaos::FTrailingDataExt TrailingData;
+
+									FTransform CurrTransform = FTransform(CollectionMassToLocal[Idx].ToMatrixWithScale() * GlobalMatrices[Idx] * ActorTransform.ToMatrixWithScale());
+									TrailingData.Location = CurrTransform.GetTranslation();
+
+									TrailingData.Velocity = (*LinearVelocity)[Idx];
+									TrailingData.AngularVelocity = (*AngularVelocity)[Idx];
+
+									TrailingData.Mass = 1.f;
+									TrailingData.BoundingboxVolume = 1000000.f;
+									TrailingData.BoundingboxExtentMin = 100.f;
+									TrailingData.BoundingboxExtentMax = 100.f;
+									TrailingData.SurfaceType = 0;
+
+									TrailingEvents.Add(TrailingData);
+								}
+							}
 						}
 					}
 				}				
