Only consider leaf level pieces of geometry for collision

diff --git a/Engine/Plugins/Experimental/ChaosNiagara/Source/ChaosNiagara/Private/NiagaraDataInterfaceGeometryCollection.cpp b/Engine/Plugins/Experimental/ChaosNiagara/Source/ChaosNiagara/Private/NiagaraDataInterfaceGeometryCollection.cpp
--- a/Engine/Plugins/Experimental/ChaosNiagara/Source/ChaosNiagara/Private/NiagaraDataInterfaceGeometryCollection.cpp
+++ b/Engine/Plugins/Experimental/ChaosNiagara/Source/ChaosNiagara/Private/NiagaraDataInterfaceGeometryCollection.cpp
@@ -146,7 +146,21 @@ void FNDIGeometryCollectionData::Init(UNiagaraDataInterfaceGeometryCollection* I
 			const TManagedArray<FBox>& BoundingBoxes = Collection->BoundingBox;
 			const TManagedArray<int32>& TransformIndex = Collection->TransformIndex;
 			const TManagedArray<FTransform>& Transforms = Interface->GeometryCollectionActor->GetGeometryCollectionComponent()->GetTransformArray();
-			int NumPieces = BoundingBoxes.Num();
+			const TManagedArray<TSet<int32>>& Children = Collection->Children;
+			const TManagedArray<int32>& TransformIndexArray = Collection->TransformIndex;
+
+			int NumPieces = 0;
+			
+			for (int i = 0; i < BoundingBoxes.Num(); ++i)
+			{
+				int32 CurrTransformIndex = TransformIndexArray[i];
+
+				// only consider leaf geometry
+				if (Collection->Children[CurrTransformIndex].Num() == 0)
+				{
+					NumPieces++;
+				}
+			}
 
 			AssetArrays = new FNDIGeometryCollectionArrays();
 			AssetArrays->Resize(NumPieces);
@@ -162,11 +176,20 @@ void FNDIGeometryCollectionData::Init(UNiagaraDataInterfaceGeometryCollection* I
 			BoundsOrigin = (FVector3f)Origin;	// LWC_TODO: Precision Loss
 			BoundsExtent = (FVector3f)Extents;	// LWC_TODO: Precision Loss
 
-			for (int i = 0; i < NumPieces; ++i)
+			int PieceIndex = 0;
+			for (int i = 0; i < BoundingBoxes.Num(); ++i)
 			{
-				FBox CurrBox = BoundingBoxes[i];
-				FVector3f BoxSize = FVector3f(CurrBox.Max - CurrBox.Min);
-				AssetArrays->BoundsBuffer[i] = FVector4f(BoxSize.X, BoxSize.Y, BoxSize.Z, 0);
+				int32 CurrTransformIndex = TransformIndexArray[i];
+
+				// only consider leaf geometry
+				if (Collection->Children[CurrTransformIndex].Num() == 0)
+				{
+					FBox CurrBox = BoundingBoxes[PieceIndex];
+					FVector3f BoxSize = FVector3f(CurrBox.Max - CurrBox.Min);
+					AssetArrays->BoundsBuffer[PieceIndex] = FVector4f(BoxSize.X, BoxSize.Y, BoxSize.Z, 0);
+
+					PieceIndex++;
+				}
 			}
 		}
 		else
@@ -199,9 +222,20 @@ void FNDIGeometryCollectionData::Update(UNiagaraDataInterfaceGeometryCollection*
 			const TManagedArray<int32>& TransformIndexArray = Collection->TransformIndex;
 			const TManagedArray<FTransform>& Transforms = Interface->GeometryCollectionActor->GetGeometryCollectionComponent()->GetTransformArray();
 			const TArray<FMatrix>& GlobalMatrices = Interface->GeometryCollectionActor->GetGeometryCollectionComponent()->GetGlobalMatrices();
+			const TManagedArray<TSet<int32>>& Children = Collection->Children;
+			
+			int NumPieces = 0;
 
-			// #todo(dmp): rename transform to pieces or something
-			int32 NumPieces = BoundingBoxes.Num();
+			for (int i = 0; i < BoundingBoxes.Num(); ++i)
+			{
+				int32 CurrTransformIndex = TransformIndexArray[i];
+
+				// only consider leaf geometry
+				if (Collection->Children[CurrTransformIndex].Num() == 0)
+				{
+					NumPieces++;
+				}
+			}
 
 			if (GlobalMatrices.Num() != Transforms.Num())
 			{
@@ -220,30 +254,37 @@ void FNDIGeometryCollectionData::Update(UNiagaraDataInterfaceGeometryCollection*
 			BoundsOrigin = (FVector3f)Origin;
 			BoundsExtent = (FVector3f)Extents;
 
-			for (int i = 0; i < NumPieces; ++i)
+			int PieceIndex = 0;
+			for (int i = 0; i < BoundingBoxes.Num(); ++i)
 			{
-				int32 TransformIndex = 3 * i;
-				AssetArrays->PrevWorldInverseTransformBuffer[TransformIndex] = AssetArrays->WorldInverseTransformBuffer[TransformIndex];
-				AssetArrays->PrevWorldInverseTransformBuffer[TransformIndex + 1] = AssetArrays->WorldInverseTransformBuffer[TransformIndex + 1];
-				AssetArrays->PrevWorldInverseTransformBuffer[TransformIndex + 2] = AssetArrays->WorldInverseTransformBuffer[TransformIndex + 2];
+				int32 CurrTransformIndex = TransformIndexArray[i];
+				if (Collection->Children[CurrTransformIndex].Num() == 0)
+				{
+					int32 TransformIndex = 3 * PieceIndex;
+					AssetArrays->PrevWorldInverseTransformBuffer[TransformIndex] = AssetArrays->WorldInverseTransformBuffer[TransformIndex];
+					AssetArrays->PrevWorldInverseTransformBuffer[TransformIndex + 1] = AssetArrays->WorldInverseTransformBuffer[TransformIndex + 1];
+					AssetArrays->PrevWorldInverseTransformBuffer[TransformIndex + 2] = AssetArrays->WorldInverseTransformBuffer[TransformIndex + 2];
 
-				AssetArrays->PrevWorldTransformBuffer[TransformIndex] = AssetArrays->WorldTransformBuffer[TransformIndex];
-				AssetArrays->PrevWorldTransformBuffer[TransformIndex + 1] = AssetArrays->WorldTransformBuffer[TransformIndex + 1];
-				AssetArrays->PrevWorldTransformBuffer[TransformIndex + 2] = AssetArrays->WorldTransformBuffer[TransformIndex + 2];
+					AssetArrays->PrevWorldTransformBuffer[TransformIndex] = AssetArrays->WorldTransformBuffer[TransformIndex];
+					AssetArrays->PrevWorldTransformBuffer[TransformIndex + 1] = AssetArrays->WorldTransformBuffer[TransformIndex + 1];
+					AssetArrays->PrevWorldTransformBuffer[TransformIndex + 2] = AssetArrays->WorldTransformBuffer[TransformIndex + 2];
 
-				FBox CurrBox = BoundingBoxes[i];
+					FBox CurrBox = BoundingBoxes[PieceIndex];
 
-				// #todo(dmp): save this somewhere in an array?
-				FVector LocalTranslation = (CurrBox.Max + CurrBox.Min) * .5;
-				FTransform LocalOffset(LocalTranslation);
-								
-				int32 CurrTransformIndex = TransformIndexArray[i];
+					// #todo(dmp): save this somewhere in an array?
+					FVector LocalTranslation = (CurrBox.Max + CurrBox.Min) * .5;
+					FTransform LocalOffset(LocalTranslation);
+
+
+
+					FMatrix44f CurrTransform = FMatrix44f(LocalOffset.ToMatrixWithScale() * GlobalMatrices[CurrTransformIndex] * ActorTransform.ToMatrixWithScale());
+					CurrTransform.To3x4MatrixTranspose(&AssetArrays->WorldTransformBuffer[TransformIndex].X);
 
-				FMatrix44f CurrTransform = FMatrix44f(LocalOffset.ToMatrixWithScale() * GlobalMatrices[CurrTransformIndex] * ActorTransform.ToMatrixWithScale());
-				CurrTransform.To3x4MatrixTranspose(&AssetArrays->WorldTransformBuffer[TransformIndex].X);
+					FMatrix44f CurrInverse = CurrTransform.Inverse();
+					CurrInverse.To3x4MatrixTranspose(&AssetArrays->WorldInverseTransformBuffer[TransformIndex].X);
 
-				FMatrix44f CurrInverse = CurrTransform.Inverse();
-				CurrInverse.To3x4MatrixTranspose(&AssetArrays->WorldInverseTransformBuffer[TransformIndex].X);
+					PieceIndex++;
+				}
 			}
 		}		
 	}
