fix bug with bounding box indices

diff --git a/Engine/Plugins/Experimental/ChaosNiagara/Source/ChaosNiagara/Private/NiagaraDataInterfaceGeometryCollection.cpp b/Engine/Plugins/Experimental/ChaosNiagara/Source/ChaosNiagara/Private/NiagaraDataInterfaceGeometryCollection.cpp
--- a/Engine/Plugins/Experimental/ChaosNiagara/Source/ChaosNiagara/Private/NiagaraDataInterfaceGeometryCollection.cpp
+++ b/Engine/Plugins/Experimental/ChaosNiagara/Source/ChaosNiagara/Private/NiagaraDataInterfaceGeometryCollection.cpp
@@ -184,7 +184,7 @@ void FNDIGeometryCollectionData::Init(UNiagaraDataInterfaceGeometryCollection* I
 				// only consider leaf geometry
 				if (Collection->Children[CurrTransformIndex].Num() == 0)
 				{
-					FBox CurrBox = BoundingBoxes[PieceIndex];
+					FBox CurrBox = BoundingBoxes[i];
 					FVector3f BoxSize = FVector3f(CurrBox.Max - CurrBox.Min);
 					AssetArrays->BoundsBuffer[PieceIndex] = FVector4f(BoxSize.X, BoxSize.Y, BoxSize.Z, 0);
 
@@ -269,7 +269,7 @@ void FNDIGeometryCollectionData::Update(UNiagaraDataInterfaceGeometryCollection*
 					AssetArrays->PrevWorldTransformBuffer[TransformIndex + 1] = AssetArrays->WorldTransformBuffer[TransformIndex + 1];
 					AssetArrays->PrevWorldTransformBuffer[TransformIndex + 2] = AssetArrays->WorldTransformBuffer[TransformIndex + 2];
 
-					FBox CurrBox = BoundingBoxes[PieceIndex];
+					FBox CurrBox = BoundingBoxes[i];
 
 					// #todo(dmp): save this somewhere in an array?
 					FVector LocalTranslation = (CurrBox.Max + CurrBox.Min) * .5;
