check for null geometry collection

diff --git a/Engine/Plugins/Experimental/ChaosNiagara/Source/ChaosNiagara/Private/NiagaraDataInterfaceGeometryCollection.cpp b/Engine/Plugins/Experimental/ChaosNiagara/Source/ChaosNiagara/Private/NiagaraDataInterfaceGeometryCollection.cpp
--- a/Engine/Plugins/Experimental/ChaosNiagara/Source/ChaosNiagara/Private/NiagaraDataInterfaceGeometryCollection.cpp
+++ b/Engine/Plugins/Experimental/ChaosNiagara/Source/ChaosNiagara/Private/NiagaraDataInterfaceGeometryCollection.cpp
@@ -139,7 +139,8 @@ void FNDIGeometryCollectionData::Init(UNiagaraDataInterfaceGeometryCollection* I
 
 	if (Interface != nullptr && SystemInstance != nullptr)
 	{		
-		if (Interface->GeometryCollectionActor != nullptr && Interface->GeometryCollectionActor->GetGeometryCollectionComponent() != nullptr)
+		if (Interface->GeometryCollectionActor != nullptr && Interface->GeometryCollectionActor->GetGeometryCollectionComponent() != nullptr &&
+			Interface->GeometryCollectionActor->GetGeometryCollectionComponent()->RestCollection != nullptr)
 		{			
 			const TSharedPtr<FGeometryCollection, ESPMode::ThreadSafe>
 				Collection = Interface->GeometryCollectionActor->GetGeometryCollectionComponent()->RestCollection->GetGeometryCollection();
