Minor dependency cleanup in Niagara.

diff --git a/Engine/Plugins/Experimental/ChaosNiagara/Source/ChaosNiagara/Private/NiagaraDataInterfaceGeometryCollection.cpp b/Engine/Plugins/Experimental/ChaosNiagara/Source/ChaosNiagara/Private/NiagaraDataInterfaceGeometryCollection.cpp
--- a/Engine/Plugins/Experimental/ChaosNiagara/Source/ChaosNiagara/Private/NiagaraDataInterfaceGeometryCollection.cpp
+++ b/Engine/Plugins/Experimental/ChaosNiagara/Source/ChaosNiagara/Private/NiagaraDataInterfaceGeometryCollection.cpp
@@ -1,22 +1,10 @@
 // Copyright Epic Games, Inc. All Rights Reserved.
 
 #include "NiagaraDataInterfaceGeometryCollection.h"
-#include "Animation/SkeletalMeshActor.h"
-#include "Components/SkeletalMeshComponent.h"
-#include "SkeletalRenderPublic.h"
-#include "SkeletalMeshTypes.h"
-#include "AnimationRuntime.h"
-#include "NiagaraShader.h"
-#include "NiagaraComponent.h"
 #include "NiagaraRenderer.h"
 #include "NiagaraSimStageData.h"
 #include "NiagaraSystemInstance.h"
 #include "ShaderParameterUtils.h"
-#include "EngineUtils.h"
-#include "Renderer/Private/ScenePrivate.h"
-#include "DistanceFieldAtlas.h"
-#include "Renderer/Private/DistanceFieldLightingShared.h"
-#include "NiagaraGpuComputeDispatchInterface.h"
 #include "GeometryCollection/GeometryCollectionComponent.h"
 
 #define LOCTEXT_NAMESPACE "NiagaraDataInterfaceGeometryCollection"
diff --git a/Engine/Plugins/FX/Niagara/Source/Niagara/Private/DataInterface/NiagaraDataInterfaceStaticMesh.cpp b/Engine/Plugins/FX/Niagara/Source/Niagara/Private/DataInterface/NiagaraDataInterfaceStaticMesh.cpp
--- a/Engine/Plugins/FX/Niagara/Source/Niagara/Private/DataInterface/NiagaraDataInterfaceStaticMesh.cpp
+++ b/Engine/Plugins/FX/Niagara/Source/Niagara/Private/DataInterface/NiagaraDataInterfaceStaticMesh.cpp
@@ -23,9 +23,6 @@
 #include "GPUSkinCache.h"
 #include "ShaderCompilerCore.h"
 
-#include "Renderer/Private/ScenePrivate.h"
-#include "DistanceFieldAtlas.h"
-#include "Renderer/Private/DistanceFieldLightingShared.h"
 #include "NiagaraGpuComputeDispatchInterface.h"
 #include "NiagaraGpuComputeDispatch.h"
 #include "NiagaraDistanceFieldParameters.h"
diff --git a/Engine/Plugins/FX/Niagara/Source/Niagara/Private/NiagaraDataInterfaceRigidMeshCollisionQuery.cpp b/Engine/Plugins/FX/Niagara/Source/Niagara/Private/NiagaraDataInterfaceRigidMeshCollisionQuery.cpp
--- a/Engine/Plugins/FX/Niagara/Source/Niagara/Private/NiagaraDataInterfaceRigidMeshCollisionQuery.cpp
+++ b/Engine/Plugins/FX/Niagara/Source/Niagara/Private/NiagaraDataInterfaceRigidMeshCollisionQuery.cpp
@@ -12,8 +12,6 @@
 #include "NiagaraSystemInstance.h"
 #include "ShaderParameterUtils.h"
 #include "EngineUtils.h"
-#include "Renderer/Private/ScenePrivate.h"
-#include "DistanceFieldAtlas.h"
 #include "NiagaraGpuComputeDispatchInterface.h"
 #include "NiagaraGpuComputeDispatch.h"
 #include "DataInterface/NiagaraDistanceFieldParameters.h"
diff --git a/Engine/Plugins/FX/Niagara/Source/Niagara/Private/NiagaraGpuComputeDispatch.cpp b/Engine/Plugins/FX/Niagara/Source/Niagara/Private/NiagaraGpuComputeDispatch.cpp
--- a/Engine/Plugins/FX/Niagara/Source/Niagara/Private/NiagaraGpuComputeDispatch.cpp
+++ b/Engine/Plugins/FX/Niagara/Source/Niagara/Private/NiagaraGpuComputeDispatch.cpp
@@ -28,9 +28,7 @@
 #include "NiagaraWorldManager.h"
 #include "RHI.h"
 #include "Runtime/Renderer/Private/SceneRendering.h"
-#include "Runtime/Renderer/Private/PostProcess/SceneRenderTargets.h"
 #include "PipelineStateCache.h"
-#include "SceneUtils.h"
 #include "ShaderParameterUtils.h"
 #include "Renderer/Private/ScenePrivate.h"
 
diff --git a/Engine/Plugins/FX/Niagara/Source/NiagaraShader/Private/NiagaraAsyncGpuTraceProviderHwrt.cpp b/Engine/Plugins/FX/Niagara/Source/NiagaraShader/Private/NiagaraAsyncGpuTraceProviderHwrt.cpp
--- a/Engine/Plugins/FX/Niagara/Source/NiagaraShader/Private/NiagaraAsyncGpuTraceProviderHwrt.cpp
+++ b/Engine/Plugins/FX/Niagara/Source/NiagaraShader/Private/NiagaraAsyncGpuTraceProviderHwrt.cpp
@@ -9,7 +9,6 @@
 #include "NiagaraRenderer.h"
 #include "NiagaraSettings.h"
 #include "Renderer/Private/ScenePrivate.h"
-#include "Renderer/Private/SceneRendering.h"
 
 static int GNiagaraAsyncGpuTraceHwrtEnabled = 1;
 static FAutoConsoleVariableRef CVarNiagaraAsyncGpuTraceHwrtEnabled(
