Chaos Niagara interface - Make sure trailing handling axcccount for deleted geometry collection component at runtime

diff --git a/Engine/Plugins/Experimental/ChaosNiagara/Source/ChaosNiagara/Private/NiagaraDataInterfaceChaosDestruction.cpp b/Engine/Plugins/Experimental/ChaosNiagara/Source/ChaosNiagara/Private/NiagaraDataInterfaceChaosDestruction.cpp
--- a/Engine/Plugins/Experimental/ChaosNiagara/Source/ChaosNiagara/Private/NiagaraDataInterfaceChaosDestruction.cpp
+++ b/Engine/Plugins/Experimental/ChaosNiagara/Source/ChaosNiagara/Private/NiagaraDataInterfaceChaosDestruction.cpp
@@ -1902,73 +1902,86 @@ void UNiagaraDataInterfaceChaosDestruction::HandleTrailingEvents(const Chaos::FT
 
 	if (GeometryCollectionComponentsFromBreaking.Num() > 0)
 	{
+		TArray<UGeometryCollectionComponent*> ComponentsToRemove;
 		for (auto& GeometryCollectionComponent : GeometryCollectionComponentsFromBreaking)
 		{
-			if (GeometryCollectionComponent && GeometryCollectionComponent->GetNotifyTrailing())
+			if (GeometryCollectionComponent)
 			{
-				PhysicalMaterial = GeometryCollectionComponent->GetPhysicalMaterial();
-				ensure(PhysicalMaterial);
-				if (PhysicalMaterial)
+				if (!GeometryCollectionComponent->IsPhysicsStateCreated())
+				{
+					ComponentsToRemove.Add(GeometryCollectionComponent);
+				}
+				else if (GeometryCollectionComponent->GetNotifyTrailing())
 				{
-					if (bApplyMaterialsFilter && !IsMaterialInFilter(PhysicalMaterial->GetFName()))
+					PhysicalMaterial = GeometryCollectionComponent->GetPhysicalMaterial();
+					ensure(PhysicalMaterial);
+					if (PhysicalMaterial)
 					{
-						continue;
+						if (bApplyMaterialsFilter && !IsMaterialInFilter(PhysicalMaterial->GetFName()))
+						{
+							continue;
+						}
 					}
-				}
 
-				const TArray<FMatrix>& GlobalMatrices = GeometryCollectionComponent->GetGlobalMatrices();
-				const FTransform ActorTransform = GeometryCollectionComponent->GetComponentToWorld();
+					const TArray<FMatrix>& GlobalMatrices = GeometryCollectionComponent->GetGlobalMatrices();
+					const FTransform ActorTransform = GeometryCollectionComponent->GetComponentToWorld();
 
-				const FGeometryDynamicCollection* DynamicCollection = GeometryCollectionComponent->GetDynamicCollection();
+					const FGeometryDynamicCollection* DynamicCollection = GeometryCollectionComponent->GetDynamicCollection();
 
-				if (DynamicCollection)
-				{
-					if (const UGeometryCollection* RestCollection = GeometryCollectionComponent->GetRestCollection())
+					if (DynamicCollection)
 					{
-						const TSharedPtr<FGeometryCollection, ESPMode::ThreadSafe> GeometryCollection = RestCollection->GetGeometryCollection();
-
-						// Get the MassToLocal transforms
-						if (GeometryCollection->HasAttribute(TEXT("MassToLocal"), FTransformCollection::TransformGroup))
+						if (const UGeometryCollection* RestCollection = GeometryCollectionComponent->GetRestCollection())
 						{
-							const TManagedArray<FTransform>& CollectionMassToLocal = GeometryCollection->GetAttribute<FTransform>(TEXT("MassToLocal"), FTransformCollection::TransformGroup);
-
-							const TManagedArray<FVector3f>* LinearVelocity = DynamicCollection->FindAttributeTyped<FVector3f>("LinearVelocity", FTransformCollection::TransformGroup);
-							const TManagedArray<FVector3f>* AngularVelocity = DynamicCollection->FindAttributeTyped<FVector3f>("AngularVelocity", FTransformCollection::TransformGroup);
+							const TSharedPtr<FGeometryCollection, ESPMode::ThreadSafe> GeometryCollection = RestCollection->GetGeometryCollection();
 
-							if (!LinearVelocity || !AngularVelocity)
+							// Get the MassToLocal transforms
+							if (GeometryCollection->HasAttribute(TEXT("MassToLocal"), FTransformCollection::TransformGroup))
 							{
-								continue;
-							}
+								const TManagedArray<FTransform>& CollectionMassToLocal = GeometryCollection->GetAttribute<FTransform>(TEXT("MassToLocal"), FTransformCollection::TransformGroup);
 
-							ensure(DynamicCollection->Active.Num() == (*LinearVelocity).Num());
-							ensure(DynamicCollection->Active.Num() == (*AngularVelocity).Num());
+								const TManagedArray<FVector3f>* LinearVelocity = DynamicCollection->FindAttributeTyped<FVector3f>("LinearVelocity", FTransformCollection::TransformGroup);
+								const TManagedArray<FVector3f>* AngularVelocity = DynamicCollection->FindAttributeTyped<FVector3f>("AngularVelocity", FTransformCollection::TransformGroup);
 
-							for (int Idx = 0; Idx < DynamicCollection->Active.Num(); ++Idx)
-							{
-								if (DynamicCollection->Active[Idx] && (*LinearVelocity)[Idx].SquaredLength() >= TrailMinSpeedToSpawnSquared)
+								if (!LinearVelocity || !AngularVelocity)
+								{
+									continue;
+								}
+
+								ensure(DynamicCollection->Active.Num() == (*LinearVelocity).Num());
+								ensure(DynamicCollection->Active.Num() == (*AngularVelocity).Num());
+
+								for (int Idx = 0; Idx < DynamicCollection->Active.Num(); ++Idx)
 								{
-									Chaos::FTrailingDataExt TrailingData;
+									if (DynamicCollection->Active[Idx] && (*LinearVelocity)[Idx].SquaredLength() >= TrailMinSpeedToSpawnSquared)
+									{
+										Chaos::FTrailingDataExt TrailingData;
 
-									FTransform CurrTransform = FTransform(CollectionMassToLocal[Idx].ToMatrixWithScale() * GlobalMatrices[Idx] * ActorTransform.ToMatrixWithScale());
-									TrailingData.Location = CurrTransform.GetTranslation();
+										FTransform CurrTransform = FTransform(CollectionMassToLocal[Idx].ToMatrixWithScale() * GlobalMatrices[Idx] * ActorTransform.ToMatrixWithScale());
+										TrailingData.Location = CurrTransform.GetTranslation();
 
-									TrailingData.Velocity = (*LinearVelocity)[Idx];
-									TrailingData.AngularVelocity = (*AngularVelocity)[Idx];
+										TrailingData.Velocity = (*LinearVelocity)[Idx];
+										TrailingData.AngularVelocity = (*AngularVelocity)[Idx];
 
-									TrailingData.Mass = 1.f;
-									TrailingData.BoundingboxVolume = 1000000.f;
-									TrailingData.BoundingboxExtentMin = 100.f;
-									TrailingData.BoundingboxExtentMax = 100.f;
-									TrailingData.SurfaceType = 0;
+										TrailingData.Mass = 1.f;
+										TrailingData.BoundingboxVolume = 1000000.f;
+										TrailingData.BoundingboxExtentMin = 100.f;
+										TrailingData.BoundingboxExtentMax = 100.f;
+										TrailingData.SurfaceType = 0;
 
-									TrailingEvents.Add(TrailingData);
+										TrailingEvents.Add(TrailingData);
+									}
 								}
 							}
 						}
-					}
-				}				
+					}				
+				}
 			}
 		}
+
+		for (UGeometryCollectionComponent* ComponentToRemove : ComponentsToRemove)
+		{
+			GeometryCollectionComponentsFromBreaking.Remove(ComponentToRemove);
+		}
 	}
 }
 
