Niagara + BP support for analytical physics fields

diff --git a/Engine/Plugins/Experimental/ChaosNiagara/Shaders/NiagaraDataInterfacePhysicsField.ush b/Engine/Plugins/Experimental/ChaosNiagara/Shaders/NiagaraDataInterfacePhysicsField.ush
--- a/Engine/Plugins/Experimental/ChaosNiagara/Shaders/NiagaraDataInterfacePhysicsField.ush
+++ b/Engine/Plugins/Experimental/ChaosNiagara/Shaders/NiagaraDataInterfacePhysicsField.ush
@@ -12,6 +12,12 @@ NiagaraDataInterfacePhysicsField.ush
  */
 
 #include "/Engine/Private/PhysicsFieldShared.ush"
+#include "/Engine/Private/PhysicsFieldEval.ush"
+
+// Global param for the world field
+int4 TargetsMapping[MAX_PHYSICS_FIELD_TARGETS];  
+int bClipmapAvailable;
+
 
 struct FDIPhysicsFieldContext
 {
@@ -49,29 +55,53 @@ DIContext.TargetCount = TargetCount_##NAME;\
 DIContext.FieldTargets = FieldTargets_##NAME;\
 DIContext.SystemLWCTile = SystemLWCTile_##NAME;\
 
-float3 DIPhysicsField_SamplePhysicsVectorField(in FDIPhysicsFieldContext DIContext, in float3 WorldPosition, in int VectorTarget)
+float3 DIPhysicsField_SamplePhysicsVectorField(in FDIPhysicsFieldContext DIContext, in float3 WorldPosition, in int VectorTarget,  in bool bUseClipmap)
 {
 	// LWC_TODO: fix this when the physics field functions are lwc ready
 	FLWCVector3 SamplePos = MakeLWCVector3(DIContext.SystemLWCTile, WorldPosition);
-	
-	return PhysicsField_SamplePhysicsVectorField(LWCHackToFloat(SamplePos), VectorTarget, DIContext.FieldTargets, DIContext.TargetCount, DIContext.ClipmapCenter, DIContext.ClipmapDistance, 
-				DIContext.ClipmapExponent, DIContext.ClipmapCount, DIContext.ClipmapResolution, DIContext.ClipmapBuffer);
+
+	if(bClipmapAvailable && bUseClipmap)
+	{
+		return PhysicsField_SamplePhysicsVectorField(LWCHackToFloat(SamplePos), VectorTarget, DIContext.FieldTargets, DIContext.TargetCount, DIContext.ClipmapCenter, DIContext.ClipmapDistance, 
+					DIContext.ClipmapExponent, DIContext.ClipmapCount, DIContext.ClipmapResolution, DIContext.ClipmapBuffer);
+	}
+	else
+	{
+		GFieldGroupThreadId = GGroupThreadId.x;
+		return PhysicsField_EvalPhysicsVectorField(LWCHackToFloat(SamplePos), TargetsMapping[VectorTarget].x);
+	}
 }
 
-float DIPhysicsField_SamplePhysicsScalarField(in FDIPhysicsFieldContext DIContext, in float3 WorldPosition, in int ScalarTarget)
+float DIPhysicsField_SamplePhysicsScalarField(in FDIPhysicsFieldContext DIContext, in float3 WorldPosition, in int ScalarTarget, in bool bUseClipmap)
 {
 	// LWC_TODO: fix this when the physics field functions are lwc ready
 	FLWCVector3 SamplePos = MakeLWCVector3(DIContext.SystemLWCTile, WorldPosition);
-	
-	return PhysicsField_SamplePhysicsScalarField(LWCHackToFloat(SamplePos), ScalarTarget, DIContext.FieldTargets, DIContext.TargetCount, DIContext.ClipmapCenter, DIContext.ClipmapDistance,
-				DIContext.ClipmapExponent, DIContext.ClipmapCount, DIContext.ClipmapResolution, DIContext.ClipmapBuffer);
+
+	if(bClipmapAvailable && bUseClipmap)
+	{
+		return PhysicsField_SamplePhysicsScalarField(LWCHackToFloat(SamplePos), ScalarTarget, DIContext.FieldTargets, DIContext.TargetCount, DIContext.ClipmapCenter, DIContext.ClipmapDistance,
+					DIContext.ClipmapExponent, DIContext.ClipmapCount, DIContext.ClipmapResolution, DIContext.ClipmapBuffer);
+	}
+	else
+	{
+		GFieldGroupThreadId = GGroupThreadId.x;
+		return PhysicsField_EvalPhysicsScalarField(LWCHackToFloat(SamplePos), TargetsMapping[ScalarTarget].y);
+	}
 }
 
-int DIPhysicsField_SamplePhysicsIntegerField(in FDIPhysicsFieldContext DIContext, in float3 WorldPosition, in int IntegerTarget)
+int DIPhysicsField_SamplePhysicsIntegerField(in FDIPhysicsFieldContext DIContext, in float3 WorldPosition, in int IntegerTarget, in bool bUseClipmap)
 {
 	// LWC_TODO: fix this when the physics field functions are lwc ready
 	FLWCVector3 SamplePos = MakeLWCVector3(DIContext.SystemLWCTile, WorldPosition);
-	
-	return PhysicsField_SamplePhysicsIntegerField(LWCHackToFloat(SamplePos), IntegerTarget, DIContext.FieldTargets, DIContext.TargetCount, DIContext.ClipmapCenter, DIContext.ClipmapDistance,
+
+	if(bClipmapAvailable && bUseClipmap)
+	{
+		return PhysicsField_SamplePhysicsIntegerField(LWCHackToFloat(SamplePos), IntegerTarget, DIContext.FieldTargets, DIContext.TargetCount, DIContext.ClipmapCenter, DIContext.ClipmapDistance,
 				DIContext.ClipmapExponent, DIContext.ClipmapCount, DIContext.ClipmapResolution, DIContext.ClipmapBuffer);
+	}
+	else
+	{
+		GFieldGroupThreadId = GGroupThreadId.x;
+		return PhysicsField_EvalPhysicsScalarField(LWCHackToFloat(SamplePos), TargetsMapping[IntegerTarget].z);
+	}
 }
\ No newline at end of file
diff --git a/Engine/Plugins/Experimental/ChaosNiagara/Source/ChaosNiagara/Classes/NiagaraDataInterfacePhysicsField.h b/Engine/Plugins/Experimental/ChaosNiagara/Source/ChaosNiagara/Classes/NiagaraDataInterfacePhysicsField.h
--- a/Engine/Plugins/Experimental/ChaosNiagara/Source/ChaosNiagara/Classes/NiagaraDataInterfacePhysicsField.h
+++ b/Engine/Plugins/Experimental/ChaosNiagara/Source/ChaosNiagara/Classes/NiagaraDataInterfacePhysicsField.h
@@ -14,6 +14,9 @@ struct FNDIFieldRenderData
 {
 	/** Field render resource for gpu */
 	class FPhysicsFieldResource* FieldResource;
+	
+	/** Time in seconds*/
+	float TimeSeconds = 0.0;
 };
 
 /** Data stored per physics asset instance on the game thread */
diff --git a/Engine/Plugins/Experimental/ChaosNiagara/Source/ChaosNiagara/Private/NiagaraDataInterfacePhysicsField.cpp b/Engine/Plugins/Experimental/ChaosNiagara/Source/ChaosNiagara/Private/NiagaraDataInterfacePhysicsField.cpp
--- a/Engine/Plugins/Experimental/ChaosNiagara/Source/ChaosNiagara/Private/NiagaraDataInterfacePhysicsField.cpp
+++ b/Engine/Plugins/Experimental/ChaosNiagara/Source/ChaosNiagara/Private/NiagaraDataInterfacePhysicsField.cpp
@@ -31,6 +31,10 @@ static const FName SamplePhysicsVectorFieldName(TEXT("SamplePhysicsVectorField")
 static const FName SamplePhysicsScalarFieldName(TEXT("SamplePhysicsScalarField"));
 static const FName SamplePhysicsIntegerFieldName(TEXT("SamplePhysicsIntegerField"));
 
+static const FName EvalPhysicsVectorFieldName(TEXT("EvalPhysicsVectorField"));
+static const FName EvalPhysicsScalarFieldName(TEXT("EvalPhysicsScalarField"));
+static const FName EvalPhysicsIntegerFieldName(TEXT("EvalPhysicsIntegerField"));
+
 static const FName GetPhysicsFieldResolutionName(TEXT("GetPhysicsFieldResolution"));
 static const FName GetPhysicsFieldBoundsName(TEXT("GetPhysicsFieldBounds"));
 
@@ -139,6 +143,13 @@ public:
 		TargetCount.Bind(ParameterMap, *ParamNames.TargetCountName);
 		FieldTargets.Bind(ParameterMap, *ParamNames.FieldTargetsName);
 		SystemLWCTile.Bind(ParameterMap, *ParamNames.SystemLWCTileName);
+
+		// Global physics field params
+		NodesParams.Bind(ParameterMap, TEXT("NodesParams"));
+		NodesOffsets.Bind(ParameterMap, TEXT("NodesOffsets"));
+		TargetsOffsets.Bind(ParameterMap, TEXT("TargetsOffsets"));
+		TargetsMapping.Bind(ParameterMap, TEXT("TargetsMapping"));
+		bClipmapAvailable.Bind(ParameterMap, TEXT("bClipmapAvailable"));
 	}
 
 	void Set(FRHICommandList& RHICmdList, const FNiagaraDataInterfaceSetArgs& Context) const
@@ -153,10 +164,31 @@ public:
 			InterfaceProxy->SystemInstancesToProxyData.Find(Context.SystemInstanceID);
 
 		TStaticArray<FIntVector4, MAX_PHYSICS_FIELD_TARGETS,16> LocalTargets;
+		TStaticArray<FIntVector4, MAX_PHYSICS_FIELD_TARGETS,16> GlobalTargets;
 		if (ProxyData != nullptr && ProxyData->FieldResource)
 		{
 			FPhysicsFieldResource* FieldResource = ProxyData->FieldResource;
 
+			const TArray<EFieldPhysicsType> VectorTypes = GetFieldTargetTypes(EFieldOutputType::Field_Output_Vector);
+			const TArray<EFieldPhysicsType> ScalarTypes = GetFieldTargetTypes(EFieldOutputType::Field_Output_Scalar);
+			const TArray<EFieldPhysicsType> IntegerTypes = GetFieldTargetTypes(EFieldOutputType::Field_Output_Integer);
+
+			int32 TypeIndex = 0;
+			for(const EFieldPhysicsType& PhysicsType : VectorTypes)
+			{
+				GlobalTargets[TypeIndex++].X = PhysicsType;
+			}
+			TypeIndex = 0;
+			for(const EFieldPhysicsType& PhysicsType : ScalarTypes)
+			{
+				GlobalTargets[TypeIndex++].Y = PhysicsType;
+			}
+			TypeIndex = 0;
+			for(const EFieldPhysicsType& PhysicsType : IntegerTypes)
+			{
+				GlobalTargets[TypeIndex++].Z = PhysicsType;
+			}
+
 			for (int32 Index = 0; Index < MAX_PHYSICS_FIELD_TARGETS; ++Index)
 			{
 				LocalTargets[Index].X = FieldResource->FieldInfos.VectorTargets[Index];
@@ -178,9 +210,15 @@ public:
 			SetShaderValue(RHICmdList, ComputeShaderRHI, ClipmapResolution, FieldResource->FieldInfos.ClipmapResolution);
 			SetShaderValue(RHICmdList, ComputeShaderRHI, ClipmapExponent, FieldResource->FieldInfos.ClipmapExponent);
 			SetShaderValue(RHICmdList, ComputeShaderRHI, ClipmapCount, FieldResource->FieldInfos.ClipmapCount);
-
 			SetShaderValue(RHICmdList, ComputeShaderRHI, TargetCount, FieldResource->FieldInfos.TargetCount);
 			SetShaderValue(RHICmdList, ComputeShaderRHI, FieldTargets, LocalTargets);
+
+			SetSRVParameter(RHICmdList, ComputeShaderRHI, NodesParams, FieldResource->NodesParams.SRV);
+			SetSRVParameter(RHICmdList, ComputeShaderRHI, NodesOffsets, FieldResource->NodesOffsets.SRV);
+			SetSRVParameter(RHICmdList, ComputeShaderRHI, TargetsOffsets, FieldResource->TargetsOffsets.SRV);
+			SetShaderValue(RHICmdList, ComputeShaderRHI, TargetsMapping, GlobalTargets);
+			SetShaderValue(RHICmdList, ComputeShaderRHI, TimeSeconds, ProxyData->TimeSeconds);
+			SetShaderValue(RHICmdList, ComputeShaderRHI, bClipmapAvailable, static_cast<int32>(FieldResource->FieldInfos.bBuildClipmap));
 		}
 		else
 		{
@@ -192,6 +230,13 @@ public:
 			SetShaderValue(RHICmdList, ComputeShaderRHI, ClipmapCount, 1);
 			SetShaderValue(RHICmdList, ComputeShaderRHI, TargetCount, 0);
 			SetShaderValue(RHICmdList, ComputeShaderRHI, FieldTargets, LocalTargets);
+			
+			SetSRVParameter(RHICmdList, ComputeShaderRHI, NodesParams, FNiagaraRenderer::GetDummyFloatBuffer());
+			SetSRVParameter(RHICmdList, ComputeShaderRHI, NodesOffsets,FNiagaraRenderer::GetDummyIntBuffer());
+			SetSRVParameter(RHICmdList, ComputeShaderRHI, TargetsOffsets, FNiagaraRenderer::GetDummyIntBuffer());
+			SetShaderValue(RHICmdList, ComputeShaderRHI, TargetsMapping, GlobalTargets);
+			SetShaderValue(RHICmdList, ComputeShaderRHI, TimeSeconds, 0.0f);
+			SetShaderValue(RHICmdList, ComputeShaderRHI, bClipmapAvailable, 0);
 		}
 		SetShaderValue(RHICmdList, ComputeShaderRHI, SystemLWCTile, Context.SystemLWCTile);
 	}
@@ -202,6 +247,7 @@ public:
 
 private:
 
+	
 	LAYOUT_FIELD(FShaderResourceParameter, ClipmapBuffer);
 	LAYOUT_FIELD(FShaderParameter, ClipmapCenter);
 	LAYOUT_FIELD(FShaderParameter, ClipmapDistance);
@@ -211,6 +257,13 @@ private:
 	LAYOUT_FIELD(FShaderParameter, TargetCount);
 	LAYOUT_FIELD(FShaderParameter, FieldTargets);
 	LAYOUT_FIELD(FShaderParameter, SystemLWCTile);
+
+	LAYOUT_FIELD(FShaderResourceParameter, NodesParams);
+	LAYOUT_FIELD(FShaderResourceParameter, NodesOffsets);
+	LAYOUT_FIELD(FShaderResourceParameter, TargetsOffsets);
+	LAYOUT_FIELD(FShaderParameter, TargetsMapping);
+	LAYOUT_FIELD(FShaderParameter, TimeSeconds);
+	LAYOUT_FIELD(FShaderParameter, bClipmapAvailable);
 };
 
 IMPLEMENT_TYPE_LAYOUT(FNDIPhysicsFieldParametersCS);
@@ -229,6 +282,7 @@ void FNDIPhysicsFieldProxy::ConsumePerInstanceDataFromGameThread(void* PerInstan
 	if (TargetData && SourceData && SourceData->FieldResource)
 	{
 		TargetData->FieldResource = SourceData->FieldResource;
+		TargetData->TimeSeconds = SourceData->TimeSeconds;
 	}
 	SourceData->~FNDIFieldRenderData();
 }
@@ -369,6 +423,42 @@ void UNiagaraDataInterfacePhysicsField::GetFunctions(TArray<FNiagaraFunctionSign
 
 		OutFunctions.Add(Sig);
 	}
+	{
+		FNiagaraFunctionSignature Sig;
+		Sig.Name = EvalPhysicsVectorFieldName;
+		Sig.bMemberFunction = true;
+		Sig.bRequiresContext = false;
+		Sig.Inputs.Add(FNiagaraVariable(FNiagaraTypeDefinition(GetClass()), TEXT("Physics Field")));
+		Sig.Inputs.Add(FNiagaraVariable(FNiagaraTypeDefinition::GetPositionDef(), TEXT("World Position")));
+		Sig.Inputs.Add(FNiagaraVariable(StaticEnum<EFieldVectorType>(), TEXT("Target Type")));
+		Sig.Outputs.Add(FNiagaraVariable(FNiagaraTypeDefinition::GetVec3Def(), TEXT("Vector Value")));
+
+		OutFunctions.Add(Sig);
+	}
+	{
+		FNiagaraFunctionSignature Sig;
+		Sig.Name = EvalPhysicsScalarFieldName;
+		Sig.bMemberFunction = true;
+		Sig.bRequiresContext = false;
+		Sig.Inputs.Add(FNiagaraVariable(FNiagaraTypeDefinition(GetClass()), TEXT("Physics Field")));
+		Sig.Inputs.Add(FNiagaraVariable(FNiagaraTypeDefinition::GetPositionDef(), TEXT("World Position")));
+		Sig.Inputs.Add(FNiagaraVariable(StaticEnum<EFieldScalarType>(), TEXT("Target Type")));
+		Sig.Outputs.Add(FNiagaraVariable(FNiagaraTypeDefinition::GetFloatDef(), TEXT("Scalar Value")));
+
+		OutFunctions.Add(Sig);
+	}
+	{
+		FNiagaraFunctionSignature Sig;
+		Sig.Name = EvalPhysicsIntegerFieldName;
+		Sig.bMemberFunction = true;
+		Sig.bRequiresContext = false;
+		Sig.Inputs.Add(FNiagaraVariable(FNiagaraTypeDefinition(GetClass()), TEXT("Physics Field")));
+		Sig.Inputs.Add(FNiagaraVariable(FNiagaraTypeDefinition::GetPositionDef(), TEXT("World Position")));
+		Sig.Inputs.Add(FNiagaraVariable(StaticEnum<EFieldIntegerType>(), TEXT("Target Type")));
+		Sig.Outputs.Add(FNiagaraVariable(FNiagaraTypeDefinition::GetIntDef(), TEXT("Integer Value")));
+
+		OutFunctions.Add(Sig);
+	}
 	{
 		FNiagaraFunctionSignature Sig;
 		Sig.Name = GetPhysicsFieldResolutionName;
@@ -415,6 +505,21 @@ void UNiagaraDataInterfacePhysicsField::GetVMExternalFunction(const FVMExternalF
 		check(BindingInfo.GetNumInputs() == 5 && BindingInfo.GetNumOutputs() == 1);
 		NDI_FUNC_BINDER(UNiagaraDataInterfacePhysicsField, SamplePhysicsIntegerField)::Bind(this, OutFunc);
 	}
+	if (BindingInfo.Name == EvalPhysicsVectorFieldName)
+	{
+		check(BindingInfo.GetNumInputs() == 5 && BindingInfo.GetNumOutputs() == 3);
+		NDI_FUNC_BINDER(UNiagaraDataInterfacePhysicsField, SamplePhysicsVectorField)::Bind(this, OutFunc);
+	}
+	else if (BindingInfo.Name == EvalPhysicsScalarFieldName)
+	{
+		check(BindingInfo.GetNumInputs() == 5 && BindingInfo.GetNumOutputs() == 1);
+		NDI_FUNC_BINDER(UNiagaraDataInterfacePhysicsField, SamplePhysicsScalarField)::Bind(this, OutFunc);
+	}
+	else if (BindingInfo.Name == EvalPhysicsIntegerFieldName)
+	{
+		check(BindingInfo.GetNumInputs() == 5 && BindingInfo.GetNumOutputs() == 1);
+		NDI_FUNC_BINDER(UNiagaraDataInterfacePhysicsField, SamplePhysicsIntegerField)::Bind(this, OutFunc);
+	}
 	else if (BindingInfo.Name == GetPhysicsFieldResolutionName)
 	{
 		check(BindingInfo.GetNumInputs() == 1 && BindingInfo.GetNumOutputs() == 3);
@@ -484,69 +589,6 @@ void UNiagaraDataInterfacePhysicsField::GetPhysicsFieldBounds(FVectorVMExternalF
 	}
 }
 
-struct FVectorFieldOperator
-{
-	static void BlendValues(const FVector& VectorValueA, const FVector& VectorValueB, FVector& VectorValueC)
-	{
-		VectorValueC = VectorValueA + VectorValueB;
-	}
-};
-
-struct FScalarFieldOperator
-{
-	static void BlendValues(const float& VectorValueA, const float& VectorValueB, float& VectorValueC)
-	{
-		VectorValueC = VectorValueA + VectorValueB;
-	}
-};
-
-struct FIntegerFieldOperator
-{
-	static void BlendValues(const int32& VectorValueA, const int32& VectorValueB, int32& VectorValueC)
-	{
-		VectorValueC = VectorValueA + VectorValueB;
-	}
-};
-
-template<typename DataType, typename BlendOperator>
-void EvaluateFieldNodes(TArray<FFieldSystemCommand>& FieldCommands, const EFieldPhysicsType FieldType, FFieldContext& FieldContext, 
-	TArray<DataType>& ResultsArray, TArray<DataType>& MaxArray)
-{
-	bool HasMatchingCommand = false;
-	if (FieldCommands.Num() > 0 && ResultsArray.Num() == MaxArray.Num())
-	{
-		SCOPE_CYCLE_COUNTER(STAT_NiagaraUpdateField_Object);
-		TFieldArrayView<DataType> ResultsView(ResultsArray, 0, ResultsArray.Num());
-
-		const float TimeSeconds = FieldContext.TimeSeconds;
-		for (int32 CommandIndex = 0; CommandIndex < FieldCommands.Num(); ++CommandIndex)
-		{
-			const FName AttributeName = FieldCommands[CommandIndex].TargetAttribute;
-			FieldContext.TimeSeconds = TimeSeconds - FieldCommands[CommandIndex].TimeCreation;
-
-			const EFieldPhysicsType CommandType = GetFieldPhysicsType(AttributeName);
-			if (CommandType == FieldType && FieldCommands[CommandIndex].RootNode.Get())
-			{
-				FFieldNode<DataType>* RootNode = static_cast<FFieldNode<DataType>*>(
-					FieldCommands[CommandIndex].RootNode.Get());
-
-				RootNode->Evaluate(FieldContext, ResultsView);
-				HasMatchingCommand = true;
-
-				for (int32 InstanceIdx = 0; InstanceIdx < MaxArray.Num(); ++InstanceIdx)
-				{
-					// TODO : First version with the add. will probably have to include an operator as a template argument 
-					BlendOperator::BlendValues(MaxArray[InstanceIdx], ResultsArray[InstanceIdx], MaxArray[InstanceIdx]);
-				} 
-			}
-		}
-	}
-	if (!HasMatchingCommand)
-	{
-		MaxArray.Init(DataType(0), ResultsArray.Num());
-	}
-}
-
 void UNiagaraDataInterfacePhysicsField::SamplePhysicsVectorField(FVectorVMExternalFunctionContext& Context)
 {
 	VectorVM::FUserPtrHandler<FNDIPhysicsFieldData> InstData(Context);
@@ -585,7 +627,7 @@ void UNiagaraDataInterfacePhysicsField::SamplePhysicsVectorField(FVectorVMExtern
 		};
 
 		const EFieldPhysicsType PhysicsType = GetFieldTargetTypes(Field_Output_Vector)[VectorTarget];
-		EvaluateFieldNodes<FVector, FVectorFieldOperator>(InstData->FieldCommands, PhysicsType, FieldContext, SampleResults, SampleMax);
+		EvaluateFieldVectorNodes(InstData->FieldCommands, PhysicsType, FieldContext, SampleResults, SampleMax);
 
 		for (int32 InstanceIdx = 0; InstanceIdx < Context.GetNumInstances(); ++InstanceIdx)
 		{
@@ -639,7 +681,7 @@ void UNiagaraDataInterfacePhysicsField::SamplePhysicsIntegerField(FVectorVMExter
 		};
 
 		const EFieldPhysicsType PhysicsType = GetFieldTargetTypes(Field_Output_Integer)[IntegerTarget];
-		EvaluateFieldNodes<int32, FIntegerFieldOperator>(InstData->FieldCommands, PhysicsType, FieldContext, SampleResults, SampleMax);
+		EvaluateFieldIntegerNodes(InstData->FieldCommands, PhysicsType, FieldContext, SampleResults, SampleMax);
 
 		for (int32 InstanceIdx = 0; InstanceIdx < Context.GetNumInstances(); ++InstanceIdx)
 		{
@@ -655,7 +697,6 @@ void UNiagaraDataInterfacePhysicsField::SamplePhysicsIntegerField(FVectorVMExter
 	}
 }
 
-
 void UNiagaraDataInterfacePhysicsField::SamplePhysicsScalarField(FVectorVMExternalFunctionContext& Context)
 {
 	VectorVM::FUserPtrHandler<FNDIPhysicsFieldData> InstData(Context);
@@ -694,7 +735,7 @@ void UNiagaraDataInterfacePhysicsField::SamplePhysicsScalarField(FVectorVMExtern
 		};
 
 		const EFieldPhysicsType PhysicsType = GetFieldTargetTypes(Field_Output_Scalar)[ScalarTarget];
-		EvaluateFieldNodes<float, FScalarFieldOperator>(InstData->FieldCommands, PhysicsType, FieldContext, SampleResults, SampleMax);
+		EvaluateFieldScalarNodes(InstData->FieldCommands, PhysicsType, FieldContext, SampleResults, SampleMax);
 
 		for (int32 InstanceIdx = 0; InstanceIdx < Context.GetNumInstances(); ++InstanceIdx)
 		{
@@ -726,7 +767,7 @@ bool UNiagaraDataInterfacePhysicsField::GetFunctionHLSL(const FNiagaraDataInterf
 		void {InstanceFunctionName}(in float3 WorldPosition, in int TargetIndex, out float3 OutFieldVector)
 		{
 			{PhysicsFieldContextName}
-			OutFieldVector = DIPhysicsField_SamplePhysicsVectorField(DIContext,WorldPosition,TargetIndex);
+			OutFieldVector = DIPhysicsField_SamplePhysicsVectorField(DIContext,WorldPosition,TargetIndex,true);
 		}
 		)");
 		OutHLSL += FString::Format(FormatSample, ArgsSample);
@@ -738,7 +779,7 @@ bool UNiagaraDataInterfacePhysicsField::GetFunctionHLSL(const FNiagaraDataInterf
 		void {InstanceFunctionName}(in float3 WorldPosition, in int TargetIndex, out float OutFieldScalar)
 		{
 			{PhysicsFieldContextName}
-			OutFieldScalar = DIPhysicsField_SamplePhysicsScalarField(DIContext,WorldPosition,TargetIndex);
+			OutFieldScalar = DIPhysicsField_SamplePhysicsScalarField(DIContext,WorldPosition,TargetIndex,true);
 		}
 		)");
 		OutHLSL += FString::Format(FormatSample, ArgsSample);
@@ -750,7 +791,43 @@ bool UNiagaraDataInterfacePhysicsField::GetFunctionHLSL(const FNiagaraDataInterf
 		void {InstanceFunctionName}(in float3 WorldPosition, in int TargetIndex, out int OutFieldInteger)
 		{
 			{PhysicsFieldContextName}
-			OutFieldInteger = DIPhysicsField_SamplePhysicsIntegerField(DIContext,WorldPosition,TargetIndex);
+			OutFieldInteger = DIPhysicsField_SamplePhysicsIntegerField(DIContext,WorldPosition,TargetIndex,true);
+		}
+		)");
+		OutHLSL += FString::Format(FormatSample, ArgsSample);
+		return true;
+	}
+	if (FunctionInfo.DefinitionName == EvalPhysicsVectorFieldName)
+	{
+		static const TCHAR* FormatSample = TEXT(R"(
+		void {InstanceFunctionName}(in float3 WorldPosition, in int TargetIndex, out float3 OutFieldVector)
+		{
+			{PhysicsFieldContextName}
+			OutFieldVector = DIPhysicsField_SamplePhysicsVectorField(DIContext,WorldPosition,TargetIndex,false);
+		}
+		)");
+		OutHLSL += FString::Format(FormatSample, ArgsSample);
+		return true;
+	}
+	else if (FunctionInfo.DefinitionName == EvalPhysicsScalarFieldName)
+	{
+		static const TCHAR* FormatSample = TEXT(R"(
+		void {InstanceFunctionName}(in float3 WorldPosition, in int TargetIndex, out float OutFieldScalar)
+		{
+			{PhysicsFieldContextName}
+			OutFieldScalar = DIPhysicsField_SamplePhysicsScalarField(DIContext,WorldPosition,TargetIndex,false);
+		}
+		)");
+		OutHLSL += FString::Format(FormatSample, ArgsSample);
+		return true;
+	}
+	else if (FunctionInfo.DefinitionName == EvalPhysicsIntegerFieldName)
+	{
+		static const TCHAR* FormatSample = TEXT(R"(
+		void {InstanceFunctionName}(in float3 WorldPosition, in int TargetIndex, out int OutFieldInteger)
+		{
+			{PhysicsFieldContextName}
+			OutFieldInteger = DIPhysicsField_SamplePhysicsIntegerField(DIContext,WorldPosition,TargetIndex,false);
 		}
 		)");
 		OutHLSL += FString::Format(FormatSample, ArgsSample);
@@ -833,6 +910,7 @@ void UNiagaraDataInterfacePhysicsField::ProvidePerInstanceDataForRenderThread(vo
 	if (GameThreadData != nullptr && RenderThreadData != nullptr)
 	{
 		RenderThreadData->FieldResource = GameThreadData->FieldResource;
+		RenderThreadData->TimeSeconds = GameThreadData->TimeSeconds;
 	}
 	check(Proxy);
 }
diff --git a/Engine/Shaders/Private/PhysicsFieldBuilder.usf b/Engine/Shaders/Private/PhysicsFieldBuilder.usf
--- a/Engine/Shaders/Private/PhysicsFieldBuilder.usf
+++ b/Engine/Shaders/Private/PhysicsFieldBuilder.usf
@@ -77,7 +77,7 @@ void BuildPhysicsFieldClipmapCS(
 	uint3 DispatchThreadId : SV_DispatchThreadID,
     uint3 GroupThreadId : SV_GroupThreadID)
 {
-	GGroupThreadId = GroupThreadId.x;
+	GFieldGroupThreadId = GroupThreadId.x;
 	if ((DispatchThreadId.x < NumCells))
 	{
 		int CellIndex = 0;
@@ -121,13 +121,13 @@ void BuildPhysicsFieldClipmapCS(
 					const int AttributeOffset = ClipmapResolution * ClipmapResolution * ClipmapResolution * ClipmapCount;
 					if (LocalIndex == 3)
 					{
-						FieldClipmap[BufferOffset] = SharedDatas[GGroupThreadId][0] * ClipmapFade;
-						FieldClipmap[BufferOffset + 1 * AttributeOffset] = SharedDatas[GGroupThreadId][1] * ClipmapFade;
-						FieldClipmap[BufferOffset + 2 * AttributeOffset] = SharedDatas[GGroupThreadId][2] * ClipmapFade;
+						FieldClipmap[BufferOffset] = SharedDatas[GFieldGroupThreadId][0] * ClipmapFade;
+						FieldClipmap[BufferOffset + 1 * AttributeOffset] = SharedDatas[GFieldGroupThreadId][1] * ClipmapFade;
+						FieldClipmap[BufferOffset + 2 * AttributeOffset] = SharedDatas[GFieldGroupThreadId][2] * ClipmapFade;
 					}
 					else if (LocalIndex == 1)
 					{
-						FieldClipmap[BufferOffset] = SharedDatas[GGroupThreadId][0] * ClipmapFade;
+						FieldClipmap[BufferOffset] = SharedDatas[GFieldGroupThreadId][0] * ClipmapFade;
 					}
 				}
 			}
diff --git a/Engine/Shaders/Private/PhysicsFieldEval.ush b/Engine/Shaders/Private/PhysicsFieldEval.ush
--- a/Engine/Shaders/Private/PhysicsFieldEval.ush
+++ b/Engine/Shaders/Private/PhysicsFieldEval.ush
@@ -25,7 +25,7 @@ float TimeSeconds;
 
 #define MAX_DATAS 16
 groupshared float SharedDatas[64][MAX_DATAS];
-static int GGroupThreadId;
+static int GFieldGroupThreadId;
 
 /* -----------------------------------------------------------------
  * Reandom Vector + Perlin Noise
@@ -247,7 +247,7 @@ float3 FieldUnRotateVectorByQuat(in float3 Vector, in float4 Quat)
 
 void EvaluateUniformInteger(in float3 SamplePosition, in int NodeOffset, inout int LocalIndex)
 {
-	SharedDatas[GGroupThreadId][LocalIndex++] = NodesParams[NodeOffset + UNIFORM_INTEGER_MAGNITUDE];
+	SharedDatas[GFieldGroupThreadId][LocalIndex++] = NodesParams[NodeOffset + UNIFORM_INTEGER_MAGNITUDE];
 }
 
 /* -----------------------------------------------------------------
@@ -275,9 +275,9 @@ void EvaluateRadialMaskInteger(in float3 SamplePosition, in int NodeOffset, inou
 											 (int)NodesParams[NodeOffset + RADIAL_MASK_INTEGER_EXTERIOR_VALUE];
 	const int MaskCondition = (int)NodesParams[NodeOffset + RADIAL_MASK_INTEGER_SET_MASK_CONDITION];
 	
-	SharedDatas[GGroupThreadId][LocalIndex++] = (MaskCondition == SET_ALWAYS) ? DeltaResult :
-					 ((MaskCondition == SET_IFF_NOT_INTERIOR) && (SharedDatas[GGroupThreadId][LocalIndex] != NodesParams[NodeOffset + RADIAL_MASK_INTEGER_INTERIOR_VALUE])) ? DeltaResult :
-					 ((MaskCondition == SET_IFF_NOT_EXTERIOR) && (SharedDatas[GGroupThreadId][LocalIndex] != NodesParams[NodeOffset + RADIAL_MASK_INTEGER_EXTERIOR_VALUE])) ? DeltaResult : SharedDatas[GGroupThreadId][LocalIndex];
+	SharedDatas[GFieldGroupThreadId][LocalIndex++] = (MaskCondition == SET_ALWAYS) ? DeltaResult :
+					 ((MaskCondition == SET_IFF_NOT_INTERIOR) && (SharedDatas[GFieldGroupThreadId][LocalIndex] != NodesParams[NodeOffset + RADIAL_MASK_INTEGER_INTERIOR_VALUE])) ? DeltaResult :
+					 ((MaskCondition == SET_IFF_NOT_EXTERIOR) && (SharedDatas[GFieldGroupThreadId][LocalIndex] != NodesParams[NodeOffset + RADIAL_MASK_INTEGER_EXTERIOR_VALUE])) ? DeltaResult : SharedDatas[GFieldGroupThreadId][LocalIndex];
 }
 
 /* -----------------------------------------------------------------
@@ -289,7 +289,7 @@ void EvaluateRadialMaskInteger(in float3 SamplePosition, in int NodeOffset, inou
 
 void EvaluateUniformScalar(in float3 SamplePosition, in int NodeOffset, inout int LocalIndex)
 {
-	SharedDatas[GGroupThreadId][LocalIndex++] = NodesParams[NodeOffset + UNIFORM_SCALAR_MAGNITUDE];
+	SharedDatas[GFieldGroupThreadId][LocalIndex++] = NodesParams[NodeOffset + UNIFORM_SCALAR_MAGNITUDE];
 }
 
 /* -----------------------------------------------------------------
@@ -386,7 +386,7 @@ void EvaluateWaveScalar(in float3 SamplePosition, in int NodeOffset, inout int L
 		Result = Magnitude * exp(-Decay * Decay);
 
 	}
-	SharedDatas[GGroupThreadId][LocalIndex++] = Result;
+	SharedDatas[GFieldGroupThreadId][LocalIndex++] = Result;
 }
 
 /* -----------------------------------------------------------------
@@ -423,7 +423,7 @@ void EvaluateRadialFalloffScalar(in float3 SamplePosition, in int NodeOffset, in
 						1.0, DeltaDistance, ScalarDatas);
 		ScalarDatas = NodesParams[NodeOffset + RADIAL_FALLOFF_MAGNITUDE] * (NodesParams[NodeOffset + RADIAL_FALLOFF_MIN_RANGE] + ScalarDatas * DeltaRange);
 	}
-	SharedDatas[GGroupThreadId][LocalIndex++] = ScalarDatas;
+	SharedDatas[GFieldGroupThreadId][LocalIndex++] = ScalarDatas;
 }
 
 /* -----------------------------------------------------------------
@@ -470,7 +470,7 @@ void EvaluatePlaneFalloffScalar(in float3 SamplePosition, in int NodeOffset, ino
 						1.0, DeltaDistance, ScalarDatas);
 		ScalarDatas = NodesParams[NodeOffset + PLANE_FALLOFF_MAGNITUDE] * (NodesParams[NodeOffset + PLANE_FALLOFF_MIN_RANGE] + ScalarDatas * DeltaRange);
 	}
-	SharedDatas[GGroupThreadId][LocalIndex++] = ScalarDatas;
+	SharedDatas[GFieldGroupThreadId][LocalIndex++] = ScalarDatas;
 }
 
 /* -----------------------------------------------------------------
@@ -534,7 +534,7 @@ void EvaluateBoxFalloffScalar(in float3 SamplePosition, in int NodeOffset, inout
 						1.0,DeltaDistance, ScalarDatas);
 		ScalarDatas = NodesParams[NodeOffset + BOX_FALLOFF_MAGNITUDE] * (NodesParams[NodeOffset + BOX_FALLOFF_MIN_RANGE] + ScalarDatas * DeltaRange);
 	}
-	SharedDatas[GGroupThreadId][LocalIndex++] = ScalarDatas;
+	SharedDatas[GFieldGroupThreadId][LocalIndex++] = ScalarDatas;
 }
 
 /* -----------------------------------------------------------------
@@ -580,7 +580,7 @@ void EvaluateNoiseScalar(in float3 SamplePosition, in int NodeOffset, inout int
 	LocalPosition = (clamp(LocalPosition, float3(-1.0,-1.0,-1.0), float3(1.0,1.0,1.0)) * 0.5 + float3(0.5,0.5,0.5)) * 255;
 
 	const float NoiseValue = PerlinNoise(LocalPosition) * 0.5 + 0.5; 
-	SharedDatas[GGroupThreadId][LocalIndex++] = NoiseValue * (NodesParams[NodeOffset + NOISE_MAX_RANGE] -
+	SharedDatas[GFieldGroupThreadId][LocalIndex++] = NoiseValue * (NodesParams[NodeOffset + NOISE_MAX_RANGE] -
 								  NodesParams[NodeOffset + NOISE_MIN_RANGE]) 
 				+ NodesParams[NodeOffset + NOISE_MIN_RANGE];
 }
@@ -603,9 +603,9 @@ void EvaluateUniformVector(in float3 SamplePosition, in int NodeOffset, inout in
 								        NodesParams[NodeOffset + UNIFORM_VECTOR_DIRECTIONY],
 								        NodesParams[NodeOffset + UNIFORM_VECTOR_DIRECTIONZ]);
 
-	SharedDatas[GGroupThreadId][LocalIndex++] = VectorDatas.x;
-	SharedDatas[GGroupThreadId][LocalIndex++] = VectorDatas.y;
-	SharedDatas[GGroupThreadId][LocalIndex++] = VectorDatas.z;
+	SharedDatas[GFieldGroupThreadId][LocalIndex++] = VectorDatas.x;
+	SharedDatas[GFieldGroupThreadId][LocalIndex++] = VectorDatas.y;
+	SharedDatas[GFieldGroupThreadId][LocalIndex++] = VectorDatas.z;
 }
 
 /* -----------------------------------------------------------------
@@ -628,9 +628,9 @@ void EvaluateRadialVector(in float3 SamplePosition, in int NodeOffset, inout int
 	const float3 VectorDatas = (DirectionLength != 0.0) ? NodesParams[NodeOffset + RADIAL_VECTOR_MAGNITUDE] *
 					RadialDirection / DirectionLength : float3(0, 0, 0);
 
-	SharedDatas[GGroupThreadId][LocalIndex++] = VectorDatas.x;
-	SharedDatas[GGroupThreadId][LocalIndex++] = VectorDatas.y;
-	SharedDatas[GGroupThreadId][LocalIndex++] = VectorDatas.z;
+	SharedDatas[GFieldGroupThreadId][LocalIndex++] = VectorDatas.x;
+	SharedDatas[GFieldGroupThreadId][LocalIndex++] = VectorDatas.y;
+	SharedDatas[GFieldGroupThreadId][LocalIndex++] = VectorDatas.z;
 }
 
 /* -----------------------------------------------------------------
@@ -648,9 +648,9 @@ void EvaluateRandomVector(in float3 SamplePosition, in int NodeOffset, inout int
 	const float3 VectorDatas = (DirectionLength != 0.0) ? NodesParams[NodeOffset + RANDOM_VECTOR_MAGNITUDE] *
 					RadialDirection / DirectionLength : float3(0, 0, 0);
 
-	SharedDatas[GGroupThreadId][LocalIndex++] = VectorDatas.x;
-	SharedDatas[GGroupThreadId][LocalIndex++] = VectorDatas.y;
-	SharedDatas[GGroupThreadId][LocalIndex++] = VectorDatas.z;
+	SharedDatas[GFieldGroupThreadId][LocalIndex++] = VectorDatas.x;
+	SharedDatas[GFieldGroupThreadId][LocalIndex++] = VectorDatas.y;
+	SharedDatas[GFieldGroupThreadId][LocalIndex++] = VectorDatas.z;
 }
 
 /* -----------------------------------------------------------------
@@ -674,9 +674,9 @@ void EvaluateSumScalar(in float3 SamplePosition, in int NodeOffset, inout int Lo
 	const bool UnitOperation = (FieldOperation == MULTIPLY_OP) || (FieldOperation == DIVIDE_OP);
 
 	int OffsetIndex = LocalIndex;
-	const float ScalarRight = HasRight ? SharedDatas[GGroupThreadId][OffsetIndex] : UnitOperation ? 1.0 : 0.0;
+	const float ScalarRight = HasRight ? SharedDatas[GFieldGroupThreadId][OffsetIndex] : UnitOperation ? 1.0 : 0.0;
 	OffsetIndex += HasRight;
-	const float ScalarLeft = HasLeft ? SharedDatas[GGroupThreadId][OffsetIndex] : UnitOperation ? 1.0 : 0.0;
+	const float ScalarLeft = HasLeft ? SharedDatas[GFieldGroupThreadId][OffsetIndex] : UnitOperation ? 1.0 : 0.0;
 
 	float ScalarDatas = 0.0;
 	if (FieldOperation == MULTIPLY_OP)
@@ -695,7 +695,7 @@ void EvaluateSumScalar(in float3 SamplePosition, in int NodeOffset, inout int Lo
 	{
 		ScalarDatas = ScalarLeft - ScalarRight;
 	}
-	SharedDatas[GGroupThreadId][LocalIndex++] = ScalarDatas * NodesParams[NodeOffset + SUM_SCALAR_MAGNITUDE];
+	SharedDatas[GFieldGroupThreadId][LocalIndex++] = ScalarDatas * NodesParams[NodeOffset + SUM_SCALAR_MAGNITUDE];
 }
 
 /* -----------------------------------------------------------------
@@ -721,11 +721,11 @@ void EvaluateSumVector(in float3 SamplePosition, in int NodeOffset, inout int Lo
 	const bool UnitOperation = (FieldOperation == MULTIPLY_OP) || (FieldOperation == DIVIDE_OP);
 
 	int OffsetIndex = LocalIndex;
-	const float ScalarMagnitude = HasScalar ? SharedDatas[GGroupThreadId][OffsetIndex] : 1.0;
+	const float ScalarMagnitude = HasScalar ? SharedDatas[GFieldGroupThreadId][OffsetIndex] : 1.0;
 	OffsetIndex += HasScalar;
-	const float3 VectorRight = HasRight ? float3(SharedDatas[GGroupThreadId][OffsetIndex], SharedDatas[GGroupThreadId][OffsetIndex + 1], SharedDatas[GGroupThreadId][OffsetIndex + 2]) : UnitOperation ? float3(1, 1, 1) : float3(0, 0, 0);
+	const float3 VectorRight = HasRight ? float3(SharedDatas[GFieldGroupThreadId][OffsetIndex], SharedDatas[GFieldGroupThreadId][OffsetIndex + 1], SharedDatas[GFieldGroupThreadId][OffsetIndex + 2]) : UnitOperation ? float3(1, 1, 1) : float3(0, 0, 0);
 	OffsetIndex += HasRight * 3;
-	const float3 VectorLeft = HasLeft ? float3(SharedDatas[GGroupThreadId][OffsetIndex], SharedDatas[GGroupThreadId][OffsetIndex + 1], SharedDatas[GGroupThreadId][OffsetIndex + 2]) : UnitOperation ? float3(1, 1, 1) : float3(0, 0, 0);
+	const float3 VectorLeft = HasLeft ? float3(SharedDatas[GFieldGroupThreadId][OffsetIndex], SharedDatas[GFieldGroupThreadId][OffsetIndex + 1], SharedDatas[GFieldGroupThreadId][OffsetIndex + 2]) : UnitOperation ? float3(1, 1, 1) : float3(0, 0, 0);
 
 	float3 VectorDatas = float3(0.0, 0.0, 0.0);
 	if (FieldOperation == MULTIPLY_OP)
@@ -747,9 +747,9 @@ void EvaluateSumVector(in float3 SamplePosition, in int NodeOffset, inout int Lo
 
 	VectorDatas *= ScalarMagnitude * NodesParams[NodeOffset + SUM_VECTOR_MAGNITUDE];
 
-	SharedDatas[GGroupThreadId][LocalIndex++] = VectorDatas.x;
-	SharedDatas[GGroupThreadId][LocalIndex++] = VectorDatas.y;
-	SharedDatas[GGroupThreadId][LocalIndex++] = VectorDatas.z;
+	SharedDatas[GFieldGroupThreadId][LocalIndex++] = VectorDatas.x;
+	SharedDatas[GFieldGroupThreadId][LocalIndex++] = VectorDatas.y;
+	SharedDatas[GFieldGroupThreadId][LocalIndex++] = VectorDatas.z;
 }
 
 /* -----------------------------------------------------------------
@@ -790,15 +790,15 @@ void EvaluateCullingInteger(in float3 SamplePosition, in int NodeOffset, inout i
 	LocalIndex -= HasScalar + HasInput;
 
 	int OffsetIndex = LocalIndex;
-	const float CullingScalar = HasScalar ? SharedDatas[GGroupThreadId][OffsetIndex] : 0.0;
+	const float CullingScalar = HasScalar ? SharedDatas[GFieldGroupThreadId][OffsetIndex] : 0.0;
 	OffsetIndex += (int)CullingScalar;
-	const int InputInteger = HasInput ? (int) SharedDatas[GGroupThreadId][OffsetIndex] : 0;
+	const int InputInteger = HasInput ? (int) SharedDatas[GFieldGroupThreadId][OffsetIndex] : 0;
 
 	const int CullingOperation = (int)NodesParams[NodeOffset + CULLING_OPERATION];
 	const bool ValidSample = ((CullingOperation == CULLING_OUTSIDE) && (CullingScalar != 0.0)) ||
 							 ((CullingOperation == CULLING_INSIDE) && (CullingScalar == 0.0));
 
-	SharedDatas[GGroupThreadId][LocalIndex++] = (float) (ValidSample ? InputInteger : 0);
+	SharedDatas[GFieldGroupThreadId][LocalIndex++] = (float) (ValidSample ? InputInteger : 0);
 }
 
 /* -----------------------------------------------------------------
@@ -814,15 +814,15 @@ void EvaluateCullingScalar(in float3 SamplePosition, in int NodeOffset, inout in
 	LocalIndex -= HasScalar + HasInput;
 
 	int OffsetIndex = LocalIndex;
-	const float CullingScalar = HasScalar ? SharedDatas[GGroupThreadId][OffsetIndex] : 0.0;
+	const float CullingScalar = HasScalar ? SharedDatas[GFieldGroupThreadId][OffsetIndex] : 0.0;
 	OffsetIndex += (int)CullingScalar;
-	const float InputScalar = HasInput ? SharedDatas[GGroupThreadId][OffsetIndex] : 0.0;
+	const float InputScalar = HasInput ? SharedDatas[GFieldGroupThreadId][OffsetIndex] : 0.0;
 
 	const int CullingOperation = (int)NodesParams[NodeOffset + CULLING_OPERATION];
 	const bool ValidSample = ((CullingOperation == CULLING_OUTSIDE) && (CullingScalar != 0.0)) ||
 							 ((CullingOperation == CULLING_INSIDE) && (CullingScalar == 0.0));
 
-	SharedDatas[GGroupThreadId][LocalIndex++] = ValidSample ? InputScalar : 0.0;
+	SharedDatas[GFieldGroupThreadId][LocalIndex++] = ValidSample ? InputScalar : 0.0;
 }
 
 /* -----------------------------------------------------------------
@@ -838,17 +838,17 @@ void EvaluateCullingVector(in float3 SamplePosition, in int NodeOffset, inout in
 	LocalIndex -= HasScalar + HasInput * 3;
 
 	int OffsetIndex = LocalIndex;
-	const float CullingScalar = HasScalar ? SharedDatas[GGroupThreadId][OffsetIndex] : 0.0;
+	const float CullingScalar = HasScalar ? SharedDatas[GFieldGroupThreadId][OffsetIndex] : 0.0;
 	OffsetIndex += (int)CullingScalar;
-	const float3 InputVector = HasInput ? float3(SharedDatas[GGroupThreadId][OffsetIndex], SharedDatas[GGroupThreadId][OffsetIndex + 1], SharedDatas[GGroupThreadId][OffsetIndex + 2]) : float3(0, 0, 0);
+	const float3 InputVector = HasInput ? float3(SharedDatas[GFieldGroupThreadId][OffsetIndex], SharedDatas[GFieldGroupThreadId][OffsetIndex + 1], SharedDatas[GFieldGroupThreadId][OffsetIndex + 2]) : float3(0, 0, 0);
 
 	const int CullingOperation = (int)NodesParams[NodeOffset + CULLING_OPERATION];
 	const bool ValidSample = ((CullingOperation == CULLING_OUTSIDE) && (CullingScalar != 0.0)) ||
 							 ((CullingOperation == CULLING_INSIDE) && (CullingScalar == 0.0));
 
-	SharedDatas[GGroupThreadId][LocalIndex++] = ValidSample ? InputVector.x : 0.0;
-	SharedDatas[GGroupThreadId][LocalIndex++] = ValidSample ? InputVector.y : 0.0;
-	SharedDatas[GGroupThreadId][LocalIndex++] = ValidSample ? InputVector.z : 0.0;
+	SharedDatas[GFieldGroupThreadId][LocalIndex++] = ValidSample ? InputVector.x : 0.0;
+	SharedDatas[GFieldGroupThreadId][LocalIndex++] = ValidSample ? InputVector.y : 0.0;
+	SharedDatas[GFieldGroupThreadId][LocalIndex++] = ValidSample ? InputVector.z : 0.0;
 }
 
 /* -----------------------------------------------------------------
@@ -977,7 +977,7 @@ void SampleFieldDatas(in float3 SamplePosition, inout int LocalIndex, in int Tar
 	[unroll]
 	for (int DatasIndex = 0; DatasIndex < MAX_DATAS; ++DatasIndex)
 	{
-		SharedDatas[GGroupThreadId][DatasIndex] = 0.0;
+		SharedDatas[GFieldGroupThreadId][DatasIndex] = 0.0;
 	}
 	[branch]
 	if (NumNodes != 0)
@@ -999,9 +999,9 @@ float3 PhysicsField_EvalPhysicsVectorField(in float3 SamplePosition, in int Targ
 
 	if(LocalIndex == 3)
 	{
-		VectorField.x = SharedDatas[GGroupThreadId][0];
-		VectorField.y = SharedDatas[GGroupThreadId][1];
-		VectorField.z = SharedDatas[GGroupThreadId][2];
+		VectorField.x = SharedDatas[GFieldGroupThreadId][0];
+		VectorField.y = SharedDatas[GFieldGroupThreadId][1];
+		VectorField.z = SharedDatas[GFieldGroupThreadId][2];
 	}
 	return VectorField;
 }
@@ -1015,7 +1015,7 @@ float PhysicsField_EvalPhysicsScalarField(in float3 SamplePosition, in int Targe
 
 	if(LocalIndex == 1)
 	{
-		ScalarField = SharedDatas[GGroupThreadId][0];
+		ScalarField = SharedDatas[GFieldGroupThreadId][0];
 	}
 	return ScalarField;
 }
@@ -1029,7 +1029,7 @@ int PhysicsField_EvalPhysicsIntegerField(in float3 SamplePosition, in int Target
 
 	if(LocalIndex == 1)
 	{
-		IntegerField = SharedDatas[GGroupThreadId][0];
+		IntegerField = SharedDatas[GFieldGroupThreadId][0];
 	}
 	return IntegerField;
 }
diff --git a/Engine/Shaders/Private/PhysicsFieldVisualizer.usf b/Engine/Shaders/Private/PhysicsFieldVisualizer.usf
--- a/Engine/Shaders/Private/PhysicsFieldVisualizer.usf
+++ b/Engine/Shaders/Private/PhysicsFieldVisualizer.usf
@@ -223,7 +223,7 @@ void PrintVectorField()
 void MainCS(uint3 DispatchThreadId : SV_DispatchThreadID,
     uint3 GroupThreadId : SV_GroupThreadID)
 {
-	GGroupThreadId = GroupThreadId.x * 16 + GroupThreadId.y * 4 + GroupThreadId.z;
+	GFieldGroupThreadId = GroupThreadId.x * 16 + GroupThreadId.y * 4 + GroupThreadId.z;
 	if(DispatchThreadId.x < VOXEL_RESOLUTION * BoundsSize && DispatchThreadId.y < VOXEL_RESOLUTION && DispatchThreadId.z < VOXEL_RESOLUTION)
 	{
 		const uint BoundsIndex = BoundsOffset + DispatchThreadId.x / VOXEL_RESOLUTION;
@@ -589,7 +589,7 @@ void PrintIntegerField()
 void MainCS(uint3 DispatchThreadId : SV_DispatchThreadID,
     uint3 GroupThreadId : SV_GroupThreadID)
 {
-	GGroupThreadId = GroupThreadId.x * 8 + GroupThreadId.y + GroupThreadId.z;
+	GFieldGroupThreadId = GroupThreadId.x * 8 + GroupThreadId.y + GroupThreadId.z;
 	if(DispatchThreadId.x < OutputResolution.x && DispatchThreadId.y < OutputResolution.y)
 	{
 		float3 RayBegin, RayEnd, RayDirection;
@@ -664,7 +664,7 @@ void MainCS(uint3 DispatchThreadId : SV_DispatchThreadID,
 void MainCS(uint3 DispatchThreadId : SV_DispatchThreadID,
     uint3 GroupThreadId : SV_GroupThreadID)
 {
-	GGroupThreadId = GroupThreadId.x * 8 + GroupThreadId.y + GroupThreadId.z;
+	GFieldGroupThreadId = GroupThreadId.x * 8 + GroupThreadId.y + GroupThreadId.z;
 	if(DispatchThreadId.x < OutputResolution.x && DispatchThreadId.y < OutputResolution.y)
 	{
 		float3 RayBegin, RayEnd, RayDirection;
diff --git a/Engine/Source/Runtime/Engine/Classes/PhysicsField/PhysicsFieldComponent.h b/Engine/Source/Runtime/Engine/Classes/PhysicsField/PhysicsFieldComponent.h
--- a/Engine/Source/Runtime/Engine/Classes/PhysicsField/PhysicsFieldComponent.h
+++ b/Engine/Source/Runtime/Engine/Classes/PhysicsField/PhysicsFieldComponent.h
@@ -296,4 +296,36 @@ public:
 	FPhysicsFieldResource* DebugResource = nullptr;
 };
 
+/** Static function with world field evaluation */
+UCLASS()
+class ENGINE_API UPhysicsFieldStatics : public UBlueprintFunctionLibrary
+{
+	GENERATED_UCLASS_BODY()
+	
+public:
+	
+	/** Evaluate the world physics vector field from BP */
+	UFUNCTION(BlueprintCallable, Category="Field", meta=(WorldContext="WorldContextObject"))
+	static FVector EvalPhysicsVectorField(const UObject* WorldContextObject, const FVector& WorldPosition, const EFieldVectorType VectorType);
+
+	/** Evaluate the world physics scalar field from BP */
+	UFUNCTION(BlueprintCallable, Category="Field", meta=(WorldContext="WorldContextObject"))
+	static float EvalPhysicsScalarField(const UObject* WorldContextObject, const FVector& WorldPosition, const EFieldScalarType ScalarType);
+
+	/** Evaluate the world physics integer field from BP */
+	UFUNCTION(BlueprintCallable, Category="Field", meta=(WorldContext="WorldContextObject"))
+	static int32 EvalPhysicsIntegerField(const UObject* WorldContextObject, const FVector& WorldPosition, const EFieldIntegerType IntegerType);
+};
+
+void ENGINE_API EvaluateFieldVectorNodes(TArray<FFieldSystemCommand>& FieldCommands, const EFieldPhysicsType FieldType, FFieldContext& FieldContext, 
+	TArray<FVector>& ResultsArray, TArray<FVector>& MaxArray);
+
+void ENGINE_API EvaluateFieldScalarNodes(TArray<FFieldSystemCommand>& FieldCommands, const EFieldPhysicsType FieldType, FFieldContext& FieldContext, 
+	TArray<float>& ResultsArray, TArray<float>& MaxArray);
+
+void ENGINE_API EvaluateFieldIntegerNodes(TArray<FFieldSystemCommand>& FieldCommands, const EFieldPhysicsType FieldType, FFieldContext& FieldContext, 
+	TArray<int32>& ResultsArray, TArray<int32>& MaxArray);
+
+
+
 
diff --git a/Engine/Source/Runtime/Engine/Private/PhysicsField/PhysicsFieldComponent.cpp b/Engine/Source/Runtime/Engine/Private/PhysicsField/PhysicsFieldComponent.cpp
--- a/Engine/Source/Runtime/Engine/Private/PhysicsField/PhysicsFieldComponent.cpp
+++ b/Engine/Source/Runtime/Engine/Private/PhysicsField/PhysicsFieldComponent.cpp
@@ -1586,6 +1586,187 @@ FPhysicsFieldSceneProxy::FPhysicsFieldSceneProxy(UPhysicsFieldComponent* Physics
 FPhysicsFieldSceneProxy::~FPhysicsFieldSceneProxy()
 {}
 
+struct FVectorFieldAddOperator
+{
+	static void BlendValues(const FVector& VectorValueA, const FVector& VectorValueB, FVector& VectorValueC)
+	{
+		VectorValueC = VectorValueA + VectorValueB;
+	}
+};
+
+struct FScalarFieldAddOperator
+{
+	static void BlendValues(const float& VectorValueA, const float& VectorValueB, float& VectorValueC)
+	{
+		VectorValueC = VectorValueA + VectorValueB;
+	}
+};
+
+struct FIntegerFieldAddOperator
+{
+	static void BlendValues(const int32& VectorValueA, const int32& VectorValueB, int32& VectorValueC)
+	{
+		VectorValueC = VectorValueA + VectorValueB;
+	}
+};
+
+template<typename DataType, typename BlendOperator>
+void EvaluateFieldNodes(TArray<FFieldSystemCommand>& FieldCommands, const EFieldPhysicsType FieldType, FFieldContext& FieldContext, 
+	TArray<DataType>& ResultsArray, TArray<DataType>& MaxArray)
+{
+	bool HasMatchingCommand = false;
+	if (FieldCommands.Num() > 0 && ResultsArray.Num() == MaxArray.Num())
+	{
+		SCOPE_CYCLE_COUNTER(STAT_NiagaraUpdateField_Object);
+		TFieldArrayView<DataType> ResultsView(ResultsArray, 0, ResultsArray.Num());
+
+		const float TimeSeconds = FieldContext.TimeSeconds;
+		for (int32 CommandIndex = 0; CommandIndex < FieldCommands.Num(); ++CommandIndex)
+		{
+			const FName AttributeName = FieldCommands[CommandIndex].TargetAttribute;
+			FieldContext.TimeSeconds = TimeSeconds - FieldCommands[CommandIndex].TimeCreation;
+
+			const EFieldPhysicsType CommandType = GetFieldPhysicsType(AttributeName);
+			if (CommandType == FieldType && FieldCommands[CommandIndex].RootNode.Get())
+			{
+				FFieldNode<DataType>* RootNode = static_cast<FFieldNode<DataType>*>(
+					FieldCommands[CommandIndex].RootNode.Get());
+
+				RootNode->Evaluate(FieldContext, ResultsView);
+				HasMatchingCommand = true;
+
+				for (int32 InstanceIdx = 0; InstanceIdx < MaxArray.Num(); ++InstanceIdx)
+				{
+					// TODO : First version with the add. will probably have to include an operator as a template argument 
+					BlendOperator::BlendValues(MaxArray[InstanceIdx], ResultsArray[InstanceIdx], MaxArray[InstanceIdx]);
+				} 
+			}
+		}
+	}
+	if (!HasMatchingCommand)
+	{
+		MaxArray.Init(DataType(0), ResultsArray.Num());
+	}
+}
+
+void EvaluateFieldVectorNodes(TArray<FFieldSystemCommand>& FieldCommands, const EFieldPhysicsType FieldType, FFieldContext& FieldContext, 
+	TArray<FVector>& ResultsArray, TArray<FVector>& MaxArray)
+{
+	EvaluateFieldNodes<FVector, FVectorFieldAddOperator>(FieldCommands,FieldType,FieldContext,ResultsArray,MaxArray);
+}
+
+void EvaluateFieldScalarNodes(TArray<FFieldSystemCommand>& FieldCommands, const EFieldPhysicsType FieldType, FFieldContext& FieldContext, 
+	TArray<float>& ResultsArray, TArray<float>& MaxArray)
+{
+	EvaluateFieldNodes<float, FScalarFieldAddOperator>(FieldCommands,FieldType,FieldContext,ResultsArray,MaxArray);
+}
+
+void EvaluateFieldIntegerNodes(TArray<FFieldSystemCommand>& FieldCommands, const EFieldPhysicsType FieldType, FFieldContext& FieldContext, 
+	TArray<int32>& ResultsArray, TArray<int32>& MaxArray)
+{
+	EvaluateFieldNodes<int32, FIntegerFieldAddOperator>(FieldCommands,FieldType,FieldContext,ResultsArray,MaxArray);
+}
+
+UPhysicsFieldStatics::UPhysicsFieldStatics(class FObjectInitializer const & ObjectInitializer)
+	: Super(ObjectInitializer)
+{
+}
+
+FVector UPhysicsFieldStatics::EvalPhysicsVectorField(const UObject* WorldContextObject, const FVector& WorldPosition, const EFieldVectorType TargetType)
+{
+	if(UWorld* ThisWorld = GEngine->GetWorldFromContextObject(WorldContextObject, EGetWorldErrorMode::LogAndReturnNull))
+	{
+		FFieldExecutionDatas ExecutionDatas;
+		ExecutionDatas.SamplePositions.Init(WorldPosition, 1);
+		ExecutionDatas.SampleIndices.Init(FFieldContextIndex(0,0),1);
+		
+		TArray<FVector>& SampleResults = ExecutionDatas.VectorResults[(uint8)EFieldCommandResultType::FinalResult];
+		SampleResults.Init(FVector::ZeroVector, 1);
+
+		TArray<FVector> SampleMax;
+		SampleMax.Init(FVector::ZeroVector, 1);
+
+		FFieldContext FieldContext{
+			ExecutionDatas,
+			FFieldContext::UniquePointerMap(),
+			ThisWorld->TimeSeconds
+		};
+
+		const EFieldPhysicsType PhysicsType = GetFieldTargetTypes(Field_Output_Vector)[TargetType];
+		UPhysicsFieldComponent* FieldComponent = ThisWorld->PhysicsField;
+		if (FieldComponent && FieldComponent->FieldInstance)
+		{
+			EvaluateFieldVectorNodes(FieldComponent->FieldInstance->FieldCommands, PhysicsType, FieldContext, SampleResults, SampleMax);
+			return SampleMax[0];
+		}
+	}
+	return FVector::Zero();
+}
+
+float UPhysicsFieldStatics::EvalPhysicsScalarField(const UObject* WorldContextObject, const FVector& WorldPosition, const EFieldScalarType TargetType)
+{
+	if(UWorld* ThisWorld = GEngine->GetWorldFromContextObject(WorldContextObject, EGetWorldErrorMode::LogAndReturnNull))
+	{
+		FFieldExecutionDatas ExecutionDatas;
+		ExecutionDatas.SamplePositions.Init(WorldPosition, 1);
+		ExecutionDatas.SampleIndices.Init(FFieldContextIndex(0,0),1);
+		
+		TArray<float>& SampleResults = ExecutionDatas.ScalarResults[(uint8)EFieldCommandResultType::FinalResult];
+		SampleResults.Init(0.0f, 1);
+
+		TArray<float> SampleMax;
+		SampleMax.Init(0.0f, 1);
+
+		FFieldContext FieldContext{
+			ExecutionDatas,
+			FFieldContext::UniquePointerMap(),
+			ThisWorld->TimeSeconds
+		};
+
+		const EFieldPhysicsType PhysicsType = GetFieldTargetTypes(Field_Output_Scalar)[TargetType];
+		UPhysicsFieldComponent* FieldComponent = ThisWorld->PhysicsField;
+		if (FieldComponent && FieldComponent->FieldInstance)
+		{
+			EvaluateFieldScalarNodes(FieldComponent->FieldInstance->FieldCommands, PhysicsType, FieldContext, SampleResults, SampleMax);
+			return SampleMax[0];
+		}
+	}
+	return 0.0;
+}
+
+int32 UPhysicsFieldStatics::EvalPhysicsIntegerField(const UObject* WorldContextObject, const FVector& WorldPosition, const EFieldIntegerType TargetType)
+{
+	if(UWorld* ThisWorld = GEngine->GetWorldFromContextObject(WorldContextObject, EGetWorldErrorMode::LogAndReturnNull))
+	{
+		FFieldExecutionDatas ExecutionDatas;
+		ExecutionDatas.SamplePositions.Init(WorldPosition, 1);
+		ExecutionDatas.SampleIndices.Init(FFieldContextIndex(0,0),1);
+		
+		TArray<int32>& SampleResults = ExecutionDatas.IntegerResults[(uint8)EFieldCommandResultType::FinalResult];
+		SampleResults.Init(0, 1);
+
+		TArray<int32> SampleMax;
+		SampleMax.Init(0, 1);
+
+		FFieldContext FieldContext{
+			ExecutionDatas,
+			FFieldContext::UniquePointerMap(),
+			ThisWorld->TimeSeconds
+		};
+
+		const EFieldPhysicsType PhysicsType = GetFieldTargetTypes(Field_Output_Scalar)[TargetType];
+		UPhysicsFieldComponent* FieldComponent = ThisWorld->PhysicsField;
+		if (FieldComponent && FieldComponent->FieldInstance)
+		{
+			EvaluateFieldIntegerNodes(FieldComponent->FieldInstance->FieldCommands, PhysicsType, FieldContext, SampleResults, SampleMax);
+			return SampleMax[0];
+		}
+	}
+	return 0;
+}
+
+
+
 
 
 
