- Physics Field Update to use Template, Shader Parameter, RDG

diff --git a/Engine/Plugins/Experimental/ChaosNiagara/Shaders/NiagaraDataInterfacePhysicsField.ush b/Engine/Plugins/Experimental/ChaosNiagara/Shaders/NiagaraDataInterfacePhysicsField.ush
--- a/Engine/Plugins/Experimental/ChaosNiagara/Shaders/NiagaraDataInterfacePhysicsField.ush
+++ b/Engine/Plugins/Experimental/ChaosNiagara/Shaders/NiagaraDataInterfacePhysicsField.ush
@@ -4,8 +4,6 @@
 NiagaraDataInterfacePhysicsField.ush 
 =============================================================================*/
 
-#pragma once
-
 /* -----------------------------------------------------------------
  * Field System constants and context
  * -----------------------------------------------------------------
@@ -15,93 +13,113 @@ NiagaraDataInterfacePhysicsField.ush
 #include "/Engine/Private/PhysicsFieldEval.ush"
 
 // Global param for the world field
-int4 TargetsMapping[MAX_PHYSICS_FIELD_TARGETS];  
-int bClipmapAvailable;
+int4			{ParameterName}_TargetMappings[MAX_PHYSICS_FIELD_TARGETS];  
+int				{ParameterName}_bClipmapAvailable;
 
+Buffer<float>	{ParameterName}_ClipmapBuffer;
+float3			{ParameterName}_ClipmapCenter;
+float			{ParameterName}_ClipmapDistance;
+int				{ParameterName}_ClipmapResolution;
+int				{ParameterName}_ClipmapExponent;
+int				{ParameterName}_ClipmapCount;
+int				{ParameterName}_TargetCount;
+int4			{ParameterName}_FieldTargets[MAX_PHYSICS_FIELD_TARGETS];
+float3			{ParameterName}_SystemLWCTile;
 
-struct FDIPhysicsFieldContext
-{
-	Buffer<float> ClipmapBuffer;
-	float3 ClipmapCenter;
-	float ClipmapDistance;
-	int ClipmapResolution;
-	int ClipmapExponent;
-	int ClipmapCount;
-	int TargetCount;
-	int4 FieldTargets[MAX_PHYSICS_FIELD_TARGETS];
-	float3 SystemLWCTile;
-};
-
-#define DIPhysicsField_DECLARE_CONSTANTS(NAME)\
-Buffer<float> ClipmapBuffer_##NAME;\
-float3 ClipmapCenter_##NAME;\
-float ClipmapDistance_##NAME;\
-int ClipmapResolution_##NAME;\
-int ClipmapExponent_##NAME;\
-int ClipmapCount_##NAME;\
-int TargetCount_##NAME;\
-int4 FieldTargets_##NAME[MAX_PHYSICS_FIELD_TARGETS];\
-float3 SystemLWCTile_##NAME;\
-
-#define DIPhysicsField_MAKE_CONTEXT(NAME)\
-FDIPhysicsFieldContext DIContext; \
-DIContext.ClipmapBuffer = ClipmapBuffer_##NAME;\
-DIContext.ClipmapCenter = ClipmapCenter_##NAME;\
-DIContext.ClipmapDistance = ClipmapDistance_##NAME;\
-DIContext.ClipmapResolution = ClipmapResolution_##NAME;\
-DIContext.ClipmapExponent = ClipmapExponent_##NAME;\
-DIContext.ClipmapCount = ClipmapCount_##NAME;\
-DIContext.TargetCount = TargetCount_##NAME;\
-DIContext.FieldTargets = FieldTargets_##NAME;\
-DIContext.SystemLWCTile = SystemLWCTile_##NAME;\
-
-float3 DIPhysicsField_SamplePhysicsVectorField(in FDIPhysicsFieldContext DIContext, in float3 WorldPosition, in int VectorTarget,  in bool bUseClipmap)
+float3 PhysicsVectorField_{ParameterName}(in float3 WorldPosition, in int TargetIndex, in bool bUseClipmap)
 {
 	// LWC_TODO: fix this when the physics field functions are lwc ready
-	FLWCVector3 SamplePos = MakeLWCVector3(DIContext.SystemLWCTile, WorldPosition);
+	FLWCVector3 SamplePos = MakeLWCVector3({ParameterName}_SystemLWCTile, WorldPosition);
 
-	if(bClipmapAvailable && bUseClipmap)
+	if ( {ParameterName}_bClipmapAvailable && bUseClipmap )
 	{
-		return PhysicsField_SamplePhysicsVectorField(LWCHackToFloat(SamplePos), VectorTarget, DIContext.FieldTargets, DIContext.TargetCount, DIContext.ClipmapCenter, DIContext.ClipmapDistance, 
-					DIContext.ClipmapExponent, DIContext.ClipmapCount, DIContext.ClipmapResolution, DIContext.ClipmapBuffer);
+		return PhysicsField_SamplePhysicsVectorField(
+				LWCHackToFloat(SamplePos), TargetIndex, {ParameterName}_FieldTargets, {ParameterName}_TargetCount, {ParameterName}_ClipmapCenter, {ParameterName}_ClipmapDistance, 
+				{ParameterName}_ClipmapExponent, {ParameterName}_ClipmapCount, {ParameterName}_ClipmapResolution, {ParameterName}_ClipmapBuffer
+			);
 	}
 	else
 	{
 		GFieldGroupThreadId = GGroupThreadId.x;
-		return PhysicsField_EvalPhysicsVectorField(LWCHackToFloat(SamplePos), TargetsMapping[VectorTarget].x);
+		return PhysicsField_EvalPhysicsVectorField(LWCHackToFloat(SamplePos), {ParameterName}_TargetMappings[TargetIndex].x);
 	}
 }
 
-float DIPhysicsField_SamplePhysicsScalarField(in FDIPhysicsFieldContext DIContext, in float3 WorldPosition, in int ScalarTarget, in bool bUseClipmap)
+float PhysicsScalarField_{ParameterName}(in float3 WorldPosition, in int TargetIndex, in bool bUseClipmap)
 {
 	// LWC_TODO: fix this when the physics field functions are lwc ready
-	FLWCVector3 SamplePos = MakeLWCVector3(DIContext.SystemLWCTile, WorldPosition);
+	FLWCVector3 SamplePos = MakeLWCVector3({ParameterName}_SystemLWCTile, WorldPosition);
 
-	if(bClipmapAvailable && bUseClipmap)
+	if ( {ParameterName}_bClipmapAvailable && bUseClipmap )
 	{
-		return PhysicsField_SamplePhysicsScalarField(LWCHackToFloat(SamplePos), ScalarTarget, DIContext.FieldTargets, DIContext.TargetCount, DIContext.ClipmapCenter, DIContext.ClipmapDistance,
-					DIContext.ClipmapExponent, DIContext.ClipmapCount, DIContext.ClipmapResolution, DIContext.ClipmapBuffer);
+		return PhysicsField_SamplePhysicsScalarField(
+					LWCHackToFloat(SamplePos), TargetIndex, {ParameterName}_FieldTargets, {ParameterName}_TargetCount, {ParameterName}_ClipmapCenter, {ParameterName}_ClipmapDistance,
+					{ParameterName}_ClipmapExponent, {ParameterName}_ClipmapCount, {ParameterName}_ClipmapResolution, {ParameterName}_ClipmapBuffer
+				);
 	}
 	else
 	{
 		GFieldGroupThreadId = GGroupThreadId.x;
-		return PhysicsField_EvalPhysicsScalarField(LWCHackToFloat(SamplePos), TargetsMapping[ScalarTarget].y);
+		return PhysicsField_EvalPhysicsScalarField(LWCHackToFloat(SamplePos), {ParameterName}_TargetMappings[TargetIndex].y);
 	}
 }
 
-int DIPhysicsField_SamplePhysicsIntegerField(in FDIPhysicsFieldContext DIContext, in float3 WorldPosition, in int IntegerTarget, in bool bUseClipmap)
+int PhysicsIntegerField_{ParameterName}(in float3 WorldPosition, in int TargetIndex, in bool bUseClipmap)
 {
 	// LWC_TODO: fix this when the physics field functions are lwc ready
-	FLWCVector3 SamplePos = MakeLWCVector3(DIContext.SystemLWCTile, WorldPosition);
+	FLWCVector3 SamplePos = MakeLWCVector3({ParameterName}_SystemLWCTile, WorldPosition);
 
-	if(bClipmapAvailable && bUseClipmap)
+	if ( {ParameterName}_bClipmapAvailable && bUseClipmap )
 	{
-		return PhysicsField_SamplePhysicsIntegerField(LWCHackToFloat(SamplePos), IntegerTarget, DIContext.FieldTargets, DIContext.TargetCount, DIContext.ClipmapCenter, DIContext.ClipmapDistance,
-				DIContext.ClipmapExponent, DIContext.ClipmapCount, DIContext.ClipmapResolution, DIContext.ClipmapBuffer);
+		return PhysicsField_SamplePhysicsIntegerField(
+					LWCHackToFloat(SamplePos), TargetIndex, {ParameterName}_FieldTargets, {ParameterName}_TargetCount, {ParameterName}_ClipmapCenter, {ParameterName}_ClipmapDistance,
+					{ParameterName}_ClipmapExponent, {ParameterName}_ClipmapCount, {ParameterName}_ClipmapResolution, {ParameterName}_ClipmapBuffer
+				);
 	}
 	else
 	{
 		GFieldGroupThreadId = GGroupThreadId.x;
-		return PhysicsField_EvalPhysicsScalarField(LWCHackToFloat(SamplePos), TargetsMapping[IntegerTarget].z);
+		return PhysicsField_EvalPhysicsScalarField(LWCHackToFloat(SamplePos), {ParameterName}_TargetMappings[TargetIndex].z);
 	}
-}
\ No newline at end of file
+}
+
+void SamplePhysicsVectorField_{ParameterName}(in float3 WorldPosition, in int TargetIndex, out float3 OutFieldVector)
+{
+	OutFieldVector = PhysicsVectorField_{ParameterName}(WorldPosition, TargetIndex, true);
+}
+
+void SamplePhysicsScalarField_{ParameterName}(in float3 WorldPosition, in int TargetIndex, out float OutFieldScalar)
+{
+	OutFieldScalar = PhysicsScalarField_{ParameterName}(WorldPosition, TargetIndex, true);
+}
+
+void SamplePhysicsIntegerField_{ParameterName}(in float3 WorldPosition, in int TargetIndex, out int OutFieldInteger)
+{
+	OutFieldInteger = PhysicsIntegerField_{ParameterName}(WorldPosition, TargetIndex, true);
+}
+
+void EvalPhysicsVectorField_{ParameterName}(in float3 WorldPosition, in int TargetIndex, out float3 OutFieldVector)
+{
+	OutFieldVector = PhysicsVectorField_{ParameterName}(WorldPosition, TargetIndex, false);
+}
+
+void EvalPhysicsScalarField_{ParameterName}(in float3 WorldPosition, in int TargetIndex, out float OutFieldScalar)
+{
+	OutFieldScalar = PhysicsScalarField_{ParameterName}(WorldPosition, TargetIndex, false);
+}
+
+void EvalPhysicsIntegerField_{ParameterName}(in float3 WorldPosition, in int TargetIndex, out int OutFieldInteger)
+{
+	OutFieldInteger = PhysicsIntegerField_{ParameterName}(WorldPosition, TargetIndex, false);
+}
+
+void GetPhysicsFieldResolution_{ParameterName}(out float3 OutTextureSize)
+{
+	OutTextureSize = {ParameterName}_ClipmapResolution;
+}
+
+void GetPhysicsFieldBounds_{ParameterName}(out float3 OutMinBounds, out float3 OutMaxBounds)
+{
+	OutMinBounds = {ParameterName}_ClipmapCenter - {ParameterName}_ClipmapDistance;
+	OutMaxBounds = {ParameterName}_ClipmapCenter + {ParameterName}_ClipmapDistance;
+}
diff --git a/Engine/Plugins/Experimental/ChaosNiagara/Source/ChaosNiagara/Classes/NiagaraDataInterfacePhysicsField.h b/Engine/Plugins/Experimental/ChaosNiagara/Source/ChaosNiagara/Classes/NiagaraDataInterfacePhysicsField.h
--- a/Engine/Plugins/Experimental/ChaosNiagara/Source/ChaosNiagara/Classes/NiagaraDataInterfacePhysicsField.h
+++ b/Engine/Plugins/Experimental/ChaosNiagara/Source/ChaosNiagara/Classes/NiagaraDataInterfacePhysicsField.h
@@ -50,9 +50,6 @@ class CHAOSNIAGARA_API UNiagaraDataInterfacePhysicsField : public UNiagaraDataIn
 	GENERATED_BODY()
 
 public:
-
-	DECLARE_NIAGARA_DI_PARAMETER();
-
 	UNiagaraDataInterfacePhysicsField();
 
 	/** UObject Interface */
@@ -71,11 +68,15 @@ public:
 
 	/** GPU simulation functionality */
 #if WITH_EDITORONLY_DATA
-	virtual void GetCommonHLSL(FString& OutHLSL) override;
+	virtual bool AppendCompileHash(FNiagaraCompileHashVisitor* InVisitor) const override;
 	virtual void GetParameterDefinitionHLSL(const FNiagaraDataInterfaceGPUParamInfo& ParamInfo, FString& OutHLSL) override;
 	virtual bool GetFunctionHLSL(const FNiagaraDataInterfaceGPUParamInfo& ParamInfo, const FNiagaraDataInterfaceGeneratedFunction& FunctionInfo, int FunctionInstanceIndex, FString& OutHLSL) override;
 	virtual bool UpgradeFunctionCall(FNiagaraFunctionSignature& FunctionSignature) override;
 #endif
+	virtual bool UseLegacyShaderBindings() const override { return false; }
+	virtual void BuildShaderParameters(FNiagaraShaderParametersBuilder& ShaderParametersBuilder) const override;
+	virtual void SetShaderParameters(const FNiagaraDataInterfaceSetShaderParametersContext& Context) const override;
+
 	virtual void ProvidePerInstanceDataForRenderThread(void* DataForRenderThread, void* PerInstanceData, const FNiagaraSystemInstanceID& SystemInstance) override;
 
 	/** Sample the vector field */
@@ -93,17 +94,6 @@ public:
 	/** Get the field bounds */
 	void GetPhysicsFieldBounds(FVectorVMExternalFunctionContext& Context);
 
-	/** Shader attributes names */
-	static const FString ClipmapBufferName;
-	static const FString ClipmapCenterName;
-	static const FString ClipmapDistanceName;
-	static const FString ClipmapResolutionName;
-	static const FString ClipmapExponentName;
-	static const FString ClipmapCountName;
-	static const FString TargetCountName;
-	static const FString FieldTargetsName;
-	static const FString SystemLWCTileName;
-
 protected:
 	/** Copy one niagara DI to this */
 	virtual bool CopyToInternal(UNiagaraDataInterface* Destination) const override;
diff --git a/Engine/Plugins/Experimental/ChaosNiagara/Source/ChaosNiagara/Private/NiagaraDataInterfacePhysicsField.cpp b/Engine/Plugins/Experimental/ChaosNiagara/Source/ChaosNiagara/Private/NiagaraDataInterfacePhysicsField.cpp
--- a/Engine/Plugins/Experimental/ChaosNiagara/Source/ChaosNiagara/Private/NiagaraDataInterfacePhysicsField.cpp
+++ b/Engine/Plugins/Experimental/ChaosNiagara/Source/ChaosNiagara/Private/NiagaraDataInterfacePhysicsField.cpp
@@ -4,6 +4,7 @@
 #include "NiagaraShader.h"
 #include "NiagaraComponent.h"
 #include "NiagaraRenderer.h"
+#include "NiagaraShaderParametersBuilder.h"
 #include "NiagaraSystemInstance.h"
 #include "ShaderParameterUtils.h"
 #include "Field/FieldSystemNodes.h"
@@ -27,56 +28,43 @@ struct FNiagaraPhysicsFieldDIFunctionVersion
 
 //------------------------------------------------------------------------------------------------------------
 
-static const FName SamplePhysicsVectorFieldName(TEXT("SamplePhysicsVectorField"));
-static const FName SamplePhysicsScalarFieldName(TEXT("SamplePhysicsScalarField"));
-static const FName SamplePhysicsIntegerFieldName(TEXT("SamplePhysicsIntegerField"));
-
-static const FName EvalPhysicsVectorFieldName(TEXT("EvalPhysicsVectorField"));
-static const FName EvalPhysicsScalarFieldName(TEXT("EvalPhysicsScalarField"));
-static const FName EvalPhysicsIntegerFieldName(TEXT("EvalPhysicsIntegerField"));
-
-static const FName GetPhysicsFieldResolutionName(TEXT("GetPhysicsFieldResolution"));
-static const FName GetPhysicsFieldBoundsName(TEXT("GetPhysicsFieldBounds"));
-
-//------------------------------------------------------------------------------------------------------------
-
-const FString UNiagaraDataInterfacePhysicsField::ClipmapBufferName(TEXT("ClipmapBuffer_"));
-const FString UNiagaraDataInterfacePhysicsField::ClipmapCenterName(TEXT("ClipmapCenter_"));
-const FString UNiagaraDataInterfacePhysicsField::ClipmapDistanceName(TEXT("ClipmapDistance_"));
-const FString UNiagaraDataInterfacePhysicsField::ClipmapResolutionName(TEXT("ClipmapResolution_"));
-const FString UNiagaraDataInterfacePhysicsField::ClipmapExponentName(TEXT("ClipmapExponent_"));
-const FString UNiagaraDataInterfacePhysicsField::ClipmapCountName(TEXT("ClipmapCount_"));
-const FString UNiagaraDataInterfacePhysicsField::TargetCountName(TEXT("TargetCount_"));
-const FString UNiagaraDataInterfacePhysicsField::FieldTargetsName(TEXT("FieldTargets_"));
-const FString UNiagaraDataInterfacePhysicsField::SystemLWCTileName(TEXT("SystemLWCTile_"));
-
-//------------------------------------------------------------------------------------------------------------
-
-struct FNDIPhysicsFieldParametersName
+namespace NDIPhysicsFieldLocal
 {
-	FNDIPhysicsFieldParametersName(const FString& Suffix)
-	{
-		ClipmapBufferName = UNiagaraDataInterfacePhysicsField::ClipmapBufferName + Suffix;
-		ClipmapCenterName = UNiagaraDataInterfacePhysicsField::ClipmapCenterName + Suffix;
-		ClipmapDistanceName = UNiagaraDataInterfacePhysicsField::ClipmapDistanceName + Suffix;
-		ClipmapResolutionName = UNiagaraDataInterfacePhysicsField::ClipmapResolutionName + Suffix;
-		ClipmapExponentName = UNiagaraDataInterfacePhysicsField::ClipmapExponentName + Suffix;
-		ClipmapCountName = UNiagaraDataInterfacePhysicsField::ClipmapCountName + Suffix;
-		TargetCountName = UNiagaraDataInterfacePhysicsField::TargetCountName + Suffix;
-		FieldTargetsName = UNiagaraDataInterfacePhysicsField::FieldTargetsName + Suffix;
-		SystemLWCTileName = UNiagaraDataInterfacePhysicsField::SystemLWCTileName + Suffix;
-	}
-
-	FString ClipmapBufferName;
-	FString ClipmapCenterName;
-	FString ClipmapDistanceName;
-	FString ClipmapResolutionName;
-	FString ClipmapExponentName;
-	FString ClipmapCountName;
-	FString TargetCountName;
-	FString FieldTargetsName;
-	FString SystemLWCTileName;
-};
+	BEGIN_SHADER_PARAMETER_STRUCT(FShaderGlobalParameters, )
+		SHADER_PARAMETER_SRV(Buffer<float>,	NodesParams)
+		SHADER_PARAMETER_SRV(Buffer<int>,	NodesOffsets)
+		SHADER_PARAMETER_SRV(Buffer<int>,	TargetsOffsets)
+		SHADER_PARAMETER(float,				TimeSeconds)
+	END_SHADER_PARAMETER_STRUCT()
+
+	BEGIN_SHADER_PARAMETER_STRUCT(FShaderInstanceParameters, )
+		SHADER_PARAMETER(int,				bClipmapAvailable)
+		SHADER_PARAMETER_ARRAY(FIntVector4, TargetMappings, [MAX_PHYSICS_FIELD_TARGETS])
+
+		SHADER_PARAMETER_SRV(Buffer<float>,	ClipmapBuffer)
+		SHADER_PARAMETER(FVector3f,			ClipmapCenter)
+		SHADER_PARAMETER(float, 			ClipmapDistance)
+		SHADER_PARAMETER(int, 				ClipmapResolution)
+		SHADER_PARAMETER(int, 				ClipmapExponent)
+		SHADER_PARAMETER(int, 				ClipmapCount)
+		SHADER_PARAMETER(int,				TargetCount)
+		SHADER_PARAMETER_ARRAY(FIntVector4, FieldTargets, [MAX_PHYSICS_FIELD_TARGETS])
+		SHADER_PARAMETER(FVector3f,			SystemLWCTile)
+	END_SHADER_PARAMETER_STRUCT()
+
+	static const TCHAR* TemplateShaderFilePath = TEXT("/Plugin/Experimental/ChaosNiagara/NiagaraDataInterfacePhysicsField.ush");
+
+	static const FName SamplePhysicsVectorFieldName("SamplePhysicsVectorField");
+	static const FName SamplePhysicsScalarFieldName("SamplePhysicsScalarField");
+	static const FName SamplePhysicsIntegerFieldName("SamplePhysicsIntegerField");
+
+	static const FName EvalPhysicsVectorFieldName("EvalPhysicsVectorField");
+	static const FName EvalPhysicsScalarFieldName("EvalPhysicsScalarField");
+	static const FName EvalPhysicsIntegerFieldName("EvalPhysicsIntegerField");
+
+	static const FName GetPhysicsFieldResolutionName("GetPhysicsFieldResolution");
+	static const FName GetPhysicsFieldBoundsName("GetPhysicsFieldBounds");
+}
 
 //------------------------------------------------------------------------------------------------------------
 
@@ -124,153 +112,6 @@ void FNDIPhysicsFieldData::Update(FNiagaraSystemInstance* SystemInstance)
 	}
 }
 
-//------------------------------------------------------------------------------------------------------------
-
-struct FNDIPhysicsFieldParametersCS : public FNiagaraDataInterfaceParametersCS
-{
-	DECLARE_TYPE_LAYOUT(FNDIPhysicsFieldParametersCS, NonVirtual);
-public:
-	void Bind(const FNiagaraDataInterfaceGPUParamInfo& ParameterInfo, const class FShaderParameterMap& ParameterMap)
-	{
-		FNDIPhysicsFieldParametersName ParamNames(*ParameterInfo.DataInterfaceHLSLSymbol);
-
-		ClipmapBuffer.Bind(ParameterMap, *ParamNames.ClipmapBufferName);
-		ClipmapCenter.Bind(ParameterMap, *ParamNames.ClipmapCenterName);
-		ClipmapDistance.Bind(ParameterMap, *ParamNames.ClipmapDistanceName);
-		ClipmapResolution.Bind(ParameterMap, *ParamNames.ClipmapResolutionName);
-		ClipmapExponent.Bind(ParameterMap, *ParamNames.ClipmapExponentName);
-		ClipmapCount.Bind(ParameterMap, *ParamNames.ClipmapCountName);
-		TargetCount.Bind(ParameterMap, *ParamNames.TargetCountName);
-		FieldTargets.Bind(ParameterMap, *ParamNames.FieldTargetsName);
-		SystemLWCTile.Bind(ParameterMap, *ParamNames.SystemLWCTileName);
-
-		// Global physics field params
-		NodesParams.Bind(ParameterMap, TEXT("NodesParams"));
-		NodesOffsets.Bind(ParameterMap, TEXT("NodesOffsets"));
-		TargetsOffsets.Bind(ParameterMap, TEXT("TargetsOffsets"));
-		TargetsMapping.Bind(ParameterMap, TEXT("TargetsMapping"));
-		bClipmapAvailable.Bind(ParameterMap, TEXT("bClipmapAvailable"));
-	}
-
-	void Set(FRHICommandList& RHICmdList, const FNiagaraDataInterfaceSetArgs& Context) const
-	{
-		check(IsInRenderingThread());
-
-		FRHIComputeShader* ComputeShaderRHI = RHICmdList.GetBoundComputeShader();
-
-		FNDIPhysicsFieldProxy* InterfaceProxy =
-			static_cast<FNDIPhysicsFieldProxy*>(Context.DataInterface);
-		FNDIFieldRenderData* ProxyData =
-			InterfaceProxy->SystemInstancesToProxyData.Find(Context.SystemInstanceID);
-
-		TStaticArray<FIntVector4, MAX_PHYSICS_FIELD_TARGETS,16> LocalTargets;
-		TStaticArray<FIntVector4, MAX_PHYSICS_FIELD_TARGETS,16> GlobalTargets;
-		if (ProxyData != nullptr && ProxyData->FieldResource)
-		{
-			FPhysicsFieldResource* FieldResource = ProxyData->FieldResource;
-
-			const TArray<EFieldPhysicsType> VectorTypes = GetFieldTargetTypes(EFieldOutputType::Field_Output_Vector);
-			const TArray<EFieldPhysicsType> ScalarTypes = GetFieldTargetTypes(EFieldOutputType::Field_Output_Scalar);
-			const TArray<EFieldPhysicsType> IntegerTypes = GetFieldTargetTypes(EFieldOutputType::Field_Output_Integer);
-
-			int32 TypeIndex = 0;
-			for(const EFieldPhysicsType& PhysicsType : VectorTypes)
-			{
-				GlobalTargets[TypeIndex++].X = PhysicsType;
-			}
-			TypeIndex = 0;
-			for(const EFieldPhysicsType& PhysicsType : ScalarTypes)
-			{
-				GlobalTargets[TypeIndex++].Y = PhysicsType;
-			}
-			TypeIndex = 0;
-			for(const EFieldPhysicsType& PhysicsType : IntegerTypes)
-			{
-				GlobalTargets[TypeIndex++].Z = PhysicsType;
-			}
-
-			for (int32 Index = 0; Index < MAX_PHYSICS_FIELD_TARGETS; ++Index)
-			{
-				LocalTargets[Index].X = FieldResource->FieldInfos.VectorTargets[Index];
-				LocalTargets[Index].Y = FieldResource->FieldInfos.ScalarTargets[Index];
-				LocalTargets[Index].Z = FieldResource->FieldInfos.IntegerTargets[Index];
-				LocalTargets[Index].W = 0; // Padding
-			}
-			if (FieldResource->FieldInfos.bBuildClipmap)
-			{
-				SetSRVParameter(RHICmdList, ComputeShaderRHI, ClipmapBuffer, FieldResource->ClipmapBuffer.SRV);
-			}
-			else
-			{
-				SetSRVParameter(RHICmdList, ComputeShaderRHI, ClipmapBuffer, FNiagaraRenderer::GetDummyFloatBuffer());
-			}
-			
-			SetShaderValue(RHICmdList, ComputeShaderRHI, ClipmapCenter, (FVector3f)FieldResource->FieldInfos.ClipmapCenter);
-			SetShaderValue(RHICmdList, ComputeShaderRHI, ClipmapDistance, FieldResource->FieldInfos.ClipmapDistance);
-			SetShaderValue(RHICmdList, ComputeShaderRHI, ClipmapResolution, FieldResource->FieldInfos.ClipmapResolution);
-			SetShaderValue(RHICmdList, ComputeShaderRHI, ClipmapExponent, FieldResource->FieldInfos.ClipmapExponent);
-			SetShaderValue(RHICmdList, ComputeShaderRHI, ClipmapCount, FieldResource->FieldInfos.ClipmapCount);
-			SetShaderValue(RHICmdList, ComputeShaderRHI, TargetCount, FieldResource->FieldInfos.TargetCount);
-			SetShaderValue(RHICmdList, ComputeShaderRHI, FieldTargets, LocalTargets);
-
-			SetSRVParameter(RHICmdList, ComputeShaderRHI, NodesParams, FieldResource->NodesParams.SRV);
-			SetSRVParameter(RHICmdList, ComputeShaderRHI, NodesOffsets, FieldResource->NodesOffsets.SRV);
-			SetSRVParameter(RHICmdList, ComputeShaderRHI, TargetsOffsets, FieldResource->TargetsOffsets.SRV);
-			SetShaderValue(RHICmdList, ComputeShaderRHI, TargetsMapping, GlobalTargets);
-			SetShaderValue(RHICmdList, ComputeShaderRHI, TimeSeconds, ProxyData->TimeSeconds);
-			SetShaderValue(RHICmdList, ComputeShaderRHI, bClipmapAvailable, static_cast<int32>(FieldResource->FieldInfos.bBuildClipmap));
-		}
-		else
-		{
-			SetSRVParameter(RHICmdList, ComputeShaderRHI, ClipmapBuffer, FNiagaraRenderer::GetDummyFloatBuffer());
-			SetShaderValue(RHICmdList, ComputeShaderRHI, ClipmapCenter, FVector3f::ZeroVector);
-			SetShaderValue(RHICmdList, ComputeShaderRHI, ClipmapDistance, 1);
-			SetShaderValue(RHICmdList, ComputeShaderRHI, ClipmapResolution, 2);
-			SetShaderValue(RHICmdList, ComputeShaderRHI, ClipmapExponent, 1);
-			SetShaderValue(RHICmdList, ComputeShaderRHI, ClipmapCount, 1);
-			SetShaderValue(RHICmdList, ComputeShaderRHI, TargetCount, 0);
-			SetShaderValue(RHICmdList, ComputeShaderRHI, FieldTargets, LocalTargets);
-			
-			SetSRVParameter(RHICmdList, ComputeShaderRHI, NodesParams, FNiagaraRenderer::GetDummyFloatBuffer());
-			SetSRVParameter(RHICmdList, ComputeShaderRHI, NodesOffsets,FNiagaraRenderer::GetDummyIntBuffer());
-			SetSRVParameter(RHICmdList, ComputeShaderRHI, TargetsOffsets, FNiagaraRenderer::GetDummyIntBuffer());
-			SetShaderValue(RHICmdList, ComputeShaderRHI, TargetsMapping, GlobalTargets);
-			SetShaderValue(RHICmdList, ComputeShaderRHI, TimeSeconds, 0.0f);
-			SetShaderValue(RHICmdList, ComputeShaderRHI, bClipmapAvailable, 0);
-		}
-		SetShaderValue(RHICmdList, ComputeShaderRHI, SystemLWCTile, Context.SystemLWCTile);
-	}
-
-	void Unset(FRHICommandList& RHICmdList, const FNiagaraDataInterfaceSetArgs& Context) const
-	{
-	}
-
-private:
-
-	
-	LAYOUT_FIELD(FShaderResourceParameter, ClipmapBuffer);
-	LAYOUT_FIELD(FShaderParameter, ClipmapCenter);
-	LAYOUT_FIELD(FShaderParameter, ClipmapDistance);
-	LAYOUT_FIELD(FShaderParameter, ClipmapResolution);
-	LAYOUT_FIELD(FShaderParameter, ClipmapExponent);
-	LAYOUT_FIELD(FShaderParameter, ClipmapCount);
-	LAYOUT_FIELD(FShaderParameter, TargetCount);
-	LAYOUT_FIELD(FShaderParameter, FieldTargets);
-	LAYOUT_FIELD(FShaderParameter, SystemLWCTile);
-
-	LAYOUT_FIELD(FShaderResourceParameter, NodesParams);
-	LAYOUT_FIELD(FShaderResourceParameter, NodesOffsets);
-	LAYOUT_FIELD(FShaderResourceParameter, TargetsOffsets);
-	LAYOUT_FIELD(FShaderParameter, TargetsMapping);
-	LAYOUT_FIELD(FShaderParameter, TimeSeconds);
-	LAYOUT_FIELD(FShaderParameter, bClipmapAvailable);
-};
-
-IMPLEMENT_TYPE_LAYOUT(FNDIPhysicsFieldParametersCS);
-
-IMPLEMENT_NIAGARA_DI_PARAMETER(UNiagaraDataInterfacePhysicsField, FNDIPhysicsFieldParametersCS);
-
-
 //------------------------------------------------------------------------------------------------------------
 
 void FNDIPhysicsFieldProxy::ConsumePerInstanceDataFromGameThread(void* PerInstanceData, const FNiagaraSystemInstanceID& Instance)
@@ -387,6 +228,8 @@ void UNiagaraDataInterfacePhysicsField::PostInitProperties()
 
 void UNiagaraDataInterfacePhysicsField::GetFunctions(TArray<FNiagaraFunctionSignature>& OutFunctions)
 {
+	using namespace NDIPhysicsFieldLocal;
+
 	{
 		FNiagaraFunctionSignature Sig;
 		Sig.Name = SamplePhysicsVectorFieldName;
@@ -490,6 +333,7 @@ DEFINE_NDI_DIRECT_FUNC_BINDER(UNiagaraDataInterfacePhysicsField, GetPhysicsField
 
 void UNiagaraDataInterfacePhysicsField::GetVMExternalFunction(const FVMExternalFunctionBindingInfo& BindingInfo, void* InstanceData, FVMExternalFunction& OutFunc)
 {
+	using namespace NDIPhysicsFieldLocal;
 	if (BindingInfo.Name == SamplePhysicsVectorFieldName)
 	{
 		check(BindingInfo.GetNumInputs() == 5 && BindingInfo.GetNumOutputs() == 3);
@@ -754,117 +598,27 @@ void UNiagaraDataInterfacePhysicsField::SamplePhysicsScalarField(FVectorVMExtern
 #if WITH_EDITORONLY_DATA
 bool UNiagaraDataInterfacePhysicsField::GetFunctionHLSL(const FNiagaraDataInterfaceGPUParamInfo& ParamInfo, const FNiagaraDataInterfaceGeneratedFunction& FunctionInfo, int FunctionInstanceIndex, FString& OutHLSL)
 {
-	FNDIPhysicsFieldParametersName ParamNames(ParamInfo.DataInterfaceHLSLSymbol);
-
-	TMap<FString, FStringFormatArg> ArgsSample = {
-		{TEXT("InstanceFunctionName"), FunctionInfo.InstanceName},
-		{TEXT("PhysicsFieldContextName"), TEXT("DIPhysicsField_MAKE_CONTEXT(") + ParamInfo.DataInterfaceHLSLSymbol + TEXT(")")},
+	using namespace NDIPhysicsFieldLocal;
+
+	static const TSet<FName> ValidGpuFunctions =
+	{
+		SamplePhysicsVectorFieldName,
+		SamplePhysicsScalarFieldName,
+		SamplePhysicsIntegerFieldName,
+		EvalPhysicsVectorFieldName,
+		EvalPhysicsScalarFieldName,
+		EvalPhysicsIntegerFieldName,
+		GetPhysicsFieldResolutionName,
+		GetPhysicsFieldBoundsName,
 	};
 
-	if (FunctionInfo.DefinitionName == SamplePhysicsVectorFieldName)
-	{
-		static const TCHAR* FormatSample = TEXT(R"(
-		void {InstanceFunctionName}(in float3 WorldPosition, in int TargetIndex, out float3 OutFieldVector)
-		{
-			{PhysicsFieldContextName}
-			OutFieldVector = DIPhysicsField_SamplePhysicsVectorField(DIContext,WorldPosition,TargetIndex,true);
-		}
-		)");
-		OutHLSL += FString::Format(FormatSample, ArgsSample);
-		return true;
-	}
-	else if (FunctionInfo.DefinitionName == SamplePhysicsScalarFieldName)
-	{
-		static const TCHAR* FormatSample = TEXT(R"(
-		void {InstanceFunctionName}(in float3 WorldPosition, in int TargetIndex, out float OutFieldScalar)
-		{
-			{PhysicsFieldContextName}
-			OutFieldScalar = DIPhysicsField_SamplePhysicsScalarField(DIContext,WorldPosition,TargetIndex,true);
-		}
-		)");
-		OutHLSL += FString::Format(FormatSample, ArgsSample);
-		return true;
-	}
-	else if (FunctionInfo.DefinitionName == SamplePhysicsIntegerFieldName)
-	{
-		static const TCHAR* FormatSample = TEXT(R"(
-		void {InstanceFunctionName}(in float3 WorldPosition, in int TargetIndex, out int OutFieldInteger)
-		{
-			{PhysicsFieldContextName}
-			OutFieldInteger = DIPhysicsField_SamplePhysicsIntegerField(DIContext,WorldPosition,TargetIndex,true);
-		}
-		)");
-		OutHLSL += FString::Format(FormatSample, ArgsSample);
-		return true;
-	}
-	if (FunctionInfo.DefinitionName == EvalPhysicsVectorFieldName)
-	{
-		static const TCHAR* FormatSample = TEXT(R"(
-		void {InstanceFunctionName}(in float3 WorldPosition, in int TargetIndex, out float3 OutFieldVector)
-		{
-			{PhysicsFieldContextName}
-			OutFieldVector = DIPhysicsField_SamplePhysicsVectorField(DIContext,WorldPosition,TargetIndex,false);
-		}
-		)");
-		OutHLSL += FString::Format(FormatSample, ArgsSample);
-		return true;
-	}
-	else if (FunctionInfo.DefinitionName == EvalPhysicsScalarFieldName)
-	{
-		static const TCHAR* FormatSample = TEXT(R"(
-		void {InstanceFunctionName}(in float3 WorldPosition, in int TargetIndex, out float OutFieldScalar)
-		{
-			{PhysicsFieldContextName}
-			OutFieldScalar = DIPhysicsField_SamplePhysicsScalarField(DIContext,WorldPosition,TargetIndex,false);
-		}
-		)");
-		OutHLSL += FString::Format(FormatSample, ArgsSample);
-		return true;
-	}
-	else if (FunctionInfo.DefinitionName == EvalPhysicsIntegerFieldName)
-	{
-		static const TCHAR* FormatSample = TEXT(R"(
-		void {InstanceFunctionName}(in float3 WorldPosition, in int TargetIndex, out int OutFieldInteger)
-		{
-			{PhysicsFieldContextName}
-			OutFieldInteger = DIPhysicsField_SamplePhysicsIntegerField(DIContext,WorldPosition,TargetIndex,false);
-		}
-		)");
-		OutHLSL += FString::Format(FormatSample, ArgsSample);
-		return true;
-	}
-	else if (FunctionInfo.DefinitionName == GetPhysicsFieldResolutionName)
-	{
-		static const TCHAR* FormatSample = TEXT(R"(
-		void {InstanceFunctionName}(out float3 OutTextureSize)
-		{
-			{PhysicsFieldContextName}
-			OutTextureSize = DIContext.ClipmapResolution;
-		}
-		)");
-		OutHLSL += FString::Format(FormatSample, ArgsSample);
-		return true;
-	}
-	else if (FunctionInfo.DefinitionName == GetPhysicsFieldBoundsName)
-	{
-		static const TCHAR* FormatSample = TEXT(R"(
-		void {InstanceFunctionName}(out float3 OutMinBounds, out float3 OutMaxBounds)
-		{
-			{PhysicsFieldContextName}
-			OutMinBounds = DIContext.ClipmapCenter - DIContext.ClipmapDistance;
-			OutMaxBounds = DIContext.ClipmapCenter + DIContext.ClipmapDistance;
-		}
-		)");
-		OutHLSL += FString::Format(FormatSample, ArgsSample);
-		return true;
-	}
-
-	OutHLSL += TEXT("\n");
-	return false;
+	return ValidGpuFunctions.Contains(FunctionInfo.DefinitionName);
 }
 
 bool UNiagaraDataInterfacePhysicsField::UpgradeFunctionCall(FNiagaraFunctionSignature& FunctionSignature)
 {
+	using namespace NDIPhysicsFieldLocal;
+
 	bool bChanged = false;
 	
 	// upgrade from lwc changes, only parameter types changed there
@@ -891,17 +645,105 @@ bool UNiagaraDataInterfacePhysicsField::UpgradeFunctionCall(FNiagaraFunctionSign
 	return bChanged;
 }
 
-void UNiagaraDataInterfacePhysicsField::GetCommonHLSL(FString& OutHLSL)
+bool UNiagaraDataInterfacePhysicsField::AppendCompileHash(FNiagaraCompileHashVisitor* InVisitor) const
 {
-	OutHLSL += TEXT("#include \"/Plugin/Experimental/ChaosNiagara/NiagaraDataInterfacePhysicsField.ush\"\n");
+	bool bSuccess = Super::AppendCompileHash(InVisitor);
+	bSuccess &= InVisitor->UpdateString(TEXT("UNiagaraDataInterfacePhysicsFieldTemplateHLSLSource"), GetShaderFileHash(NDIPhysicsFieldLocal::TemplateShaderFilePath, EShaderPlatform::SP_PCD3D_SM5).ToString());
+	bSuccess &= InVisitor->UpdateShaderParameters<NDIPhysicsFieldLocal::FShaderInstanceParameters>();
+	bSuccess &= InVisitor->UpdateShaderParameters<NDIPhysicsFieldLocal::FShaderGlobalParameters>();
+	return bSuccess;
 }
 
 void UNiagaraDataInterfacePhysicsField::GetParameterDefinitionHLSL(const FNiagaraDataInterfaceGPUParamInfo& ParamInfo, FString& OutHLSL)
 {
-	OutHLSL += TEXT("DIPhysicsField_DECLARE_CONSTANTS(") + ParamInfo.DataInterfaceHLSLSymbol + TEXT(")\n");
+	TMap<FString, FStringFormatArg> TemplateArgs =
+	{
+		{TEXT("ParameterName"),	ParamInfo.DataInterfaceHLSLSymbol},
+	};
+
+	FString TemplateFile;
+	LoadShaderSourceFile(NDIPhysicsFieldLocal::TemplateShaderFilePath, EShaderPlatform::SP_PCD3D_SM5, &TemplateFile, nullptr);
+	OutHLSL += FString::Format(*TemplateFile, TemplateArgs);
 }
 #endif
 
+void UNiagaraDataInterfacePhysicsField::BuildShaderParameters(FNiagaraShaderParametersBuilder& ShaderParametersBuilder) const
+{
+	ShaderParametersBuilder.AddNestedStruct<NDIPhysicsFieldLocal::FShaderInstanceParameters>();
+	ShaderParametersBuilder.AddIncludedStruct<NDIPhysicsFieldLocal::FShaderGlobalParameters>();
+}
+
+void UNiagaraDataInterfacePhysicsField::SetShaderParameters(const FNiagaraDataInterfaceSetShaderParametersContext& Context) const
+{
+	FNDIPhysicsFieldProxy& InterfaceProxy = Context.GetProxy<FNDIPhysicsFieldProxy>();
+	FNDIFieldRenderData* ProxyData = InterfaceProxy.SystemInstancesToProxyData.Find(Context.GetSystemInstanceID());
+
+	NDIPhysicsFieldLocal::FShaderInstanceParameters* ShaderInstanceParameters = Context.GetParameterNestedStruct<NDIPhysicsFieldLocal::FShaderInstanceParameters>();
+	NDIPhysicsFieldLocal::FShaderGlobalParameters* ShaderGlobalParameters = Context.GetParameterIncludedStruct<NDIPhysicsFieldLocal::FShaderGlobalParameters>();
+
+	if (ProxyData != nullptr && ProxyData->FieldResource)
+	{
+		FPhysicsFieldResource* FieldResource = ProxyData->FieldResource;
+
+		// Global Parameters
+		ShaderGlobalParameters->NodesParams = FieldResource->NodesParams.SRV;
+		ShaderGlobalParameters->NodesOffsets = FieldResource->NodesOffsets.SRV;
+		ShaderGlobalParameters->TargetsOffsets = FieldResource->TargetsOffsets.SRV;
+		ShaderGlobalParameters->TimeSeconds = ProxyData->TimeSeconds;
+
+		// Instance Parameters
+		const TArray<EFieldPhysicsType> VectorTypes = GetFieldTargetTypes(EFieldOutputType::Field_Output_Vector);
+		const TArray<EFieldPhysicsType> ScalarTypes = GetFieldTargetTypes(EFieldOutputType::Field_Output_Scalar);
+		const TArray<EFieldPhysicsType> IntegerTypes = GetFieldTargetTypes(EFieldOutputType::Field_Output_Integer);
+
+		for (int32 i=0; i < MAX_PHYSICS_FIELD_TARGETS; ++i)
+		{
+			ShaderInstanceParameters->TargetMappings[i].X = i < VectorTypes.Num() ? VectorTypes[i] : 0;
+			ShaderInstanceParameters->TargetMappings[i].Y = i < ScalarTypes.Num() ? ScalarTypes[i] : 0;
+			ShaderInstanceParameters->TargetMappings[i].Z = i < IntegerTypes.Num() ? IntegerTypes[i] : 0;
+			ShaderInstanceParameters->TargetMappings[i].W = 0;
+
+			ShaderInstanceParameters->FieldTargets[i].X = FieldResource->FieldInfos.VectorTargets[i];
+			ShaderInstanceParameters->FieldTargets[i].Y = FieldResource->FieldInfos.ScalarTargets[i];
+			ShaderInstanceParameters->FieldTargets[i].Z = FieldResource->FieldInfos.IntegerTargets[i];
+			ShaderInstanceParameters->FieldTargets[i].W = 0;
+		}
+
+		ShaderInstanceParameters->bClipmapAvailable = static_cast<int32>(FieldResource->FieldInfos.bBuildClipmap);
+		ShaderInstanceParameters->ClipmapBuffer		= FieldResource->FieldInfos.bBuildClipmap ? FieldResource->ClipmapBuffer.SRV.GetReference() : FNiagaraRenderer::GetDummyFloatBuffer();
+		ShaderInstanceParameters->ClipmapCenter		= FVector3f(FieldResource->FieldInfos.ClipmapCenter);
+		ShaderInstanceParameters->ClipmapDistance	= FieldResource->FieldInfos.ClipmapDistance;
+		ShaderInstanceParameters->ClipmapResolution	= FieldResource->FieldInfos.ClipmapResolution;
+		ShaderInstanceParameters->ClipmapExponent	= FieldResource->FieldInfos.ClipmapExponent;
+		ShaderInstanceParameters->ClipmapCount		= FieldResource->FieldInfos.ClipmapCount;
+		ShaderInstanceParameters->TargetCount		= FieldResource->FieldInfos.TargetCount;
+	}
+	else
+	{
+		// Global Parameters
+		ShaderGlobalParameters->NodesParams		= FNiagaraRenderer::GetDummyFloatBuffer();
+		ShaderGlobalParameters->NodesOffsets	= FNiagaraRenderer::GetDummyIntBuffer();
+		ShaderGlobalParameters->TargetsOffsets	= FNiagaraRenderer::GetDummyIntBuffer();
+		ShaderGlobalParameters->TimeSeconds		= 0.0f;
+
+		// Instance Parameters
+		ShaderInstanceParameters->bClipmapAvailable	= 0;
+		ShaderInstanceParameters->ClipmapBuffer		= FNiagaraRenderer::GetDummyFloatBuffer();
+		ShaderInstanceParameters->ClipmapCenter		= FVector3f::ZeroVector;
+		ShaderInstanceParameters->ClipmapDistance	= 1;
+		ShaderInstanceParameters->ClipmapResolution	= 2;
+		ShaderInstanceParameters->ClipmapExponent	= 1;
+		ShaderInstanceParameters->ClipmapCount		= 1;
+		ShaderInstanceParameters->TargetCount		= 0;
+		for (int i = 0; i < MAX_PHYSICS_FIELD_TARGETS; ++i)
+		{
+			ShaderInstanceParameters->TargetMappings[i] = FIntVector4(0, 0, 0, 0);
+			ShaderInstanceParameters->FieldTargets[i] = FIntVector4(0, 0, 0, 0);
+		}
+	}
+	ShaderInstanceParameters->SystemLWCTile = Context.GetSystemLWCTile();
+}
+
 void UNiagaraDataInterfacePhysicsField::ProvidePerInstanceDataForRenderThread(void* DataForRenderThread, void* PerInstanceData, const FNiagaraSystemInstanceID& SystemInstance)
 {
 	FNDIPhysicsFieldData* GameThreadData = static_cast<FNDIPhysicsFieldData*>(PerInstanceData);
@@ -915,4 +757,4 @@ void UNiagaraDataInterfacePhysicsField::ProvidePerInstanceDataForRenderThread(vo
 	check(Proxy);
 }
 
-#undef LOCTEXT_NAMESPACE
\ No newline at end of file
+#undef LOCTEXT_NAMESPACE
diff --git a/Engine/Plugins/FX/Niagara/Source/Niagara/Classes/NiagaraDataInterface.h b/Engine/Plugins/FX/Niagara/Source/Niagara/Classes/NiagaraDataInterface.h
--- a/Engine/Plugins/FX/Niagara/Source/Niagara/Classes/NiagaraDataInterface.h
+++ b/Engine/Plugins/FX/Niagara/Source/Niagara/Classes/NiagaraDataInterface.h
@@ -389,6 +389,7 @@ struct NIAGARA_API FNiagaraDataInterfaceSetShaderParametersContext
 	const FNiagaraSimStageData& GetSimStageData() const { return SimStageData; }
 
 	FNiagaraSystemInstanceID GetSystemInstanceID() const;
+	FVector3f GetSystemLWCTile() const;
 	bool IsResourceBound(const void* ResourceAddress) const;
 	bool IsParameterBound(const void* ParameterAddress) const;
 	template<typename T> bool IsStructBound(const T* StructAddress) const { return IsStructBoundInternal(StructAddress, sizeof(T)); }
diff --git a/Engine/Plugins/FX/Niagara/Source/Niagara/Private/NiagaraDataInterface.cpp b/Engine/Plugins/FX/Niagara/Source/Niagara/Private/NiagaraDataInterface.cpp
--- a/Engine/Plugins/FX/Niagara/Source/Niagara/Private/NiagaraDataInterface.cpp
+++ b/Engine/Plugins/FX/Niagara/Source/Niagara/Private/NiagaraDataInterface.cpp
@@ -265,6 +265,11 @@ FNiagaraSystemInstanceID FNiagaraDataInterfaceSetShaderParametersContext::GetSys
 	return SystemTick.SystemInstanceID;
 }
 
+FVector3f FNiagaraDataInterfaceSetShaderParametersContext::GetSystemLWCTile() const
+{
+	return SystemTick.SystemGpuComputeProxy->GetSystemLWCTile();
+}
+
 bool FNiagaraDataInterfaceSetShaderParametersContext::IsResourceBound(const void* ResourceAddress) const
 {
 	const uint16 ByteOffset = uint16(uintptr_t(ResourceAddress) - uintptr_t(BaseParameters));
diff --git a/Engine/Plugins/FX/Niagara/Source/Niagara/Private/NiagaraDataInterfaceCamera.cpp b/Engine/Plugins/FX/Niagara/Source/Niagara/Private/NiagaraDataInterfaceCamera.cpp
--- a/Engine/Plugins/FX/Niagara/Source/Niagara/Private/NiagaraDataInterfaceCamera.cpp
+++ b/Engine/Plugins/FX/Niagara/Source/Niagara/Private/NiagaraDataInterfaceCamera.cpp
@@ -372,7 +372,7 @@ void UNiagaraDataInterfaceCamera::BuildShaderParameters(FNiagaraShaderParameters
 void UNiagaraDataInterfaceCamera::SetShaderParameters(const FNiagaraDataInterfaceSetShaderParametersContext& Context) const
 {
 	FShaderParameters* ShaderParameters = Context.GetParameterNestedStruct<FShaderParameters>();
-	ShaderParameters->SystemLWCTile = Context.GetSystemTick().SystemGpuComputeProxy->GetSystemLWCTile();
+	ShaderParameters->SystemLWCTile = Context.GetSystemLWCTile();
 }
 
 DEFINE_NDI_DIRECT_FUNC_BINDER(UNiagaraDataInterfaceCamera, GetClosestParticles);
diff --git a/Engine/Plugins/FX/Niagara/Source/Niagara/Private/NiagaraDataInterfaceCollisionQuery.cpp b/Engine/Plugins/FX/Niagara/Source/Niagara/Private/NiagaraDataInterfaceCollisionQuery.cpp
--- a/Engine/Plugins/FX/Niagara/Source/Niagara/Private/NiagaraDataInterfaceCollisionQuery.cpp
+++ b/Engine/Plugins/FX/Niagara/Source/Niagara/Private/NiagaraDataInterfaceCollisionQuery.cpp
@@ -541,7 +541,7 @@ void UNiagaraDataInterfaceCollisionQuery::BuildShaderParameters(FNiagaraShaderPa
 void UNiagaraDataInterfaceCollisionQuery::SetShaderParameters(const FNiagaraDataInterfaceSetShaderParametersContext& Context) const
 {
 	FShaderParameters* ShaderParameters = Context.GetParameterNestedStruct<FShaderParameters>();
-	ShaderParameters->SystemLWCTile = Context.GetSystemTick().SystemGpuComputeProxy->GetSystemLWCTile();	//-OPT: Seems like a lot of DIs might want tile
+	ShaderParameters->SystemLWCTile = Context.GetSystemLWCTile();	//-OPT: Seems like a lot of DIs might want tile
 
 	FGlobalDistanceFieldParameters2* ShaderGDFParameters = Context.GetParameterIncludedStruct<FGlobalDistanceFieldParameters2>();
 	if (Context.IsStructBound(ShaderGDFParameters))
diff --git a/Engine/Plugins/FX/Niagara/Source/Niagara/Private/NiagaraDataInterfaceOcclusion.cpp b/Engine/Plugins/FX/Niagara/Source/Niagara/Private/NiagaraDataInterfaceOcclusion.cpp
--- a/Engine/Plugins/FX/Niagara/Source/Niagara/Private/NiagaraDataInterfaceOcclusion.cpp
+++ b/Engine/Plugins/FX/Niagara/Source/Niagara/Private/NiagaraDataInterfaceOcclusion.cpp
@@ -154,7 +154,7 @@ void UNiagaraDataInterfaceOcclusion::BuildShaderParameters(FNiagaraShaderParamet
 void UNiagaraDataInterfaceOcclusion::SetShaderParameters(const FNiagaraDataInterfaceSetShaderParametersContext& Context) const
 {
 	FShaderParameters* ShaderParameters = Context.GetParameterNestedStruct<FShaderParameters>();
-	ShaderParameters->SystemLWCTile = Context.GetSystemTick().SystemGpuComputeProxy->GetSystemLWCTile();
+	ShaderParameters->SystemLWCTile = Context.GetSystemLWCTile();
 }
 
 #undef LOCTEXT_NAMESPACE
diff --git a/Engine/Source/Runtime/Core/Private/UObject/UE5MainStreamObjectVersion.cpp b/Engine/Source/Runtime/Core/Private/UObject/UE5MainStreamObjectVersion.cpp
--- a/Engine/Source/Runtime/Core/Private/UObject/UE5MainStreamObjectVersion.cpp
+++ b/Engine/Source/Runtime/Core/Private/UObject/UE5MainStreamObjectVersion.cpp
@@ -14,7 +14,7 @@ TMap<FGuid, FGuid> FUE5MainStreamObjectVersion::GetSystemGuids()
 	SystemGuids.Add(DevGuids.MATERIALSHADERMAP_DERIVEDDATA_VER, FGuid("65C01D817C9A4EEFAE9E988D41A1F3DD"));
 	SystemGuids.Add(DevGuids.NANITE_DERIVEDDATA_VER, FGuid("7A3EA93546624597BEBE2ADFDE3A7628"));
 	SystemGuids.Add(DevGuids.NIAGARASHADERMAP_DERIVEDDATA_VER, FGuid("8A37C45D24F2423CBE5F8F371DE33575"));
-	SystemGuids.Add(DevGuids.Niagara_LatestScriptCompileVersion, FGuid("F8EAE1D4DF12704EB35657412E887A7D"));
+	SystemGuids.Add(DevGuids.Niagara_LatestScriptCompileVersion, FGuid("C7D06326D0D2074A9B458A402FD9CAC4"));
 	SystemGuids.Add(DevGuids.SkeletalMeshDerivedDataVersion, FGuid("FB7B2FB546DE49B2A2F88B7F69763CEE"));
 	SystemGuids.Add(DevGuids.STATICMESH_DERIVEDDATA_VER, FGuid("98F7D79A5811013825E75BBCC41ED3E9"));
 	SystemGuids.Add(DevGuids.POSESEARCHDB_DERIVEDDATA_VER, FGuid("4E595C2AC5E947D6BA9ABC874353E5BC"));
