- Convert GeometryCollection / ChaosDestruction DIs to RDG

diff --git a/Engine/Plugins/Experimental/ChaosNiagara/Shaders/NiagaraDataInterfaceChaosDestruction.ush b/Engine/Plugins/Experimental/ChaosNiagara/Shaders/NiagaraDataInterfaceChaosDestruction.ush
new file mode 100644
--- /dev/null
+++ b/Engine/Plugins/Experimental/ChaosNiagara/Shaders/NiagaraDataInterfaceChaosDestruction.ush
@@ -0,0 +1,180 @@
+// Copyright Epic Games, Inc. All Rights Reserved.
+
+Buffer<float3> {ParameterName}_PositionBuffer;
+Buffer<float3> {ParameterName}_VelocityBuffer;
+Buffer<float>  {ParameterName}_ExtentMinBuffer;
+Buffer<float>  {ParameterName}_ExtentMaxBuffer;
+Buffer<float>  {ParameterName}_VolumeBuffer;
+Buffer<int>    {ParameterName}_SolverIDBuffer; // NOTE(mv): Not used?
+Buffer<float>  {ParameterName}_DensityBuffer;
+Buffer<float>  {ParameterName}_FrictionBuffer;
+Buffer<float>  {ParameterName}_RestitutionBuffer;
+Buffer<int>    {ParameterName}_SurfaceTypeBuffer;
+Buffer<float4> {ParameterName}_ColorBuffer;
+		
+Buffer<float3> {ParameterName}_IncomingLocationBuffer;
+Buffer<float3> {ParameterName}_IncomingAccumulatedImpulseBuffer;
+Buffer<float3> {ParameterName}_IncomingNormalBuffer;
+Buffer<float3> {ParameterName}_IncomingVelocity1Buffer;
+Buffer<float3> {ParameterName}_IncomingVelocity2Buffer;
+Buffer<float3> {ParameterName}_IncomingAngularVelocity1Buffer;
+Buffer<float3> {ParameterName}_IncomingAngularVelocity2Buffer;
+Buffer<float>  {ParameterName}_IncomingMass1Buffer;
+Buffer<float>  {ParameterName}_IncomingMass2Buffer;
+Buffer<float>  {ParameterName}_IncomingTimeBuffer;
+
+Buffer<float3> {ParameterName}_TransformTranslationBuffer;
+Buffer<float4> {ParameterName}_TransformRotationBuffer;
+Buffer<float3> {ParameterName}_TransformScaleBuffer;
+Buffer<float3> {ParameterName}_BoundsBuffer;
+
+// NOTE(mv): Not implemented in the CPU-side functionality yet. 
+//           Returns 0 in GetPointType instead.
+//           
+// Buffer<int> PointTypeBuffer;
+
+int {ParameterName}_LastSpawnedPointID;
+
+float {ParameterName}_SolverTime;
+
+
+void GetPosition_{ParameterName}(in int ParticleID, out float3 Out_Position) 
+{
+	ParticleID -= {ParameterName}_LastSpawnedPointID + 1;
+	Out_Position = {ParameterName}_PositionBuffer[ParticleID];
+}
+
+void GetNormal_{ParameterName}(in int ParticleID, out float3 Out_Normal)
+{
+	ParticleID -= {ParameterName}_LastSpawnedPointID + 1;
+	Out_Normal = {ParameterName}_IncomingNormalBuffer[ParticleID];
+}
+
+void GetVelocity_{ParameterName}(in int ParticleID, out float3 Out_Velocity) 
+{
+	ParticleID -= {ParameterName}_LastSpawnedPointID + 1;
+	Out_Velocity = {ParameterName}_VelocityBuffer[ParticleID];
+}
+
+void GetAngularVelocity_{ParameterName}(in int ParticleID, out float3 Out_AngularVelocity) 
+{
+	ParticleID -= {ParameterName}_LastSpawnedPointID + 1;
+	Out_AngularVelocity = {ParameterName}_IncomingAngularVelocity1Buffer[ParticleID];
+}
+
+void GetExtentMin_{ParameterName}(in int ParticleID, out float Out_ExtentMin) 
+{
+	ParticleID -= {ParameterName}_LastSpawnedPointID + 1;
+	Out_ExtentMin = {ParameterName}_ExtentMinBuffer[ParticleID];
+}
+
+void GetExtentMax_{ParameterName}(in int ParticleID, out float Out_ExtentMax) 
+{
+	ParticleID -= {ParameterName}_LastSpawnedPointID + 1;
+	Out_ExtentMax = {ParameterName}_ExtentMaxBuffer[ParticleID];
+}
+
+void GetVolume_{ParameterName}(in int ParticleID, out float Out_Volume) 
+{
+	ParticleID -= {ParameterName}_LastSpawnedPointID + 1;
+	Out_Volume = {ParameterName}_VolumeBuffer[ParticleID];
+}
+
+void GetParticleIdsToSpawnAtTime_{ParameterName}(in float Time, out int Out_Min, out int Out_Max, out int Out_Count) 
+{
+	// This function cannot be called on the GPU, as all spawn scripts are run on the CPU..
+	// TODO: Find a way to warn/error about this.
+	Out_Count = 0;
+	Out_Min = 0;
+	Out_Max = 0;
+}
+
+void GetPointType_{ParameterName}(in int ParticleID, out int Out_PointType) 
+{
+	// NOTE(mv): Not yet part of the CPU functionality.
+	Out_PointType = 0;
+}
+
+void GetColor_{ParameterName}(in int ParticleID, out float4 Out_Color) 
+{
+	ParticleID -= {ParameterName}_LastSpawnedPointID + 1;
+	Out_Color = {ParameterName}_ColorBuffer[ParticleID];
+}
+
+void GetSolverTime_{ParameterName}(out float Out_SolverTime) 
+{
+	Out_SolverTime = {ParameterName}_SolverTime;
+}
+
+void GetDensity_{ParameterName}(in int ParticleID, out float Out_Density) 
+{
+	ParticleID -= {ParameterName}_LastSpawnedPointID + 1;
+	Out_Density = {ParameterName}_DensityBuffer[ParticleID];
+}
+
+void GetFriction_{ParameterName}(in int ParticleID, out float Out_Friction) 
+{
+	ParticleID -= {ParameterName}_LastSpawnedPointID + 1;
+	Out_Friction = {ParameterName}_FrictionBuffer[ParticleID];
+}
+
+void GetRestitution_{ParameterName}(in int ParticleID, out float Out_Restitution) 
+{
+	ParticleID -= {ParameterName}_LastSpawnedPointID + 1;
+	Out_Restitution = {ParameterName}_RestitutionBuffer[ParticleID];
+}
+
+void GetTransform_{ParameterName}(in int ParticleID, out float3 Out_Translation, out float4 Out_Rotation, out float3 Out_Scale) 
+{
+	ParticleID -= {ParameterName}_LastSpawnedPointID + 1;
+	Out_Translation = {ParameterName}_TransformTranslationBuffer[ParticleID];
+	Out_Rotation = {ParameterName}_TransformRotationBuffer[ParticleID];
+	Out_Scale = {ParameterName}_TransformScaleBuffer[ParticleID];
+}
+
+void GetSize_{ParameterName}(in int ParticleID, out float3 Out_Size) 
+{
+	ParticleID -= {ParameterName}_LastSpawnedPointID + 1;
+	Out_Size = {ParameterName}_BoundsBuffer[ParticleID];
+}
+
+void GetSurfaceType_{ParameterName}(in int ParticleID, out int Out_SurfaceType) 
+{
+	ParticleID -= {ParameterName}_LastSpawnedPointID + 1;
+	Out_SurfaceType = {ParameterName}_SurfaceTypeBuffer[ParticleID];
+}
+
+void GetCollisionData_{ParameterName}(in int ParticleID, out float3 Out_Location, out float3 Out_AccumulatedImpulse, out float3 Out_Normal, out float3 Out_Velocity1, out float3 Out_Velocity2, out float3 Out_AngularVelocity1, out float3 Out_AngularVelocity2, out float Out_Mass1, out float Out_Mass2, out float Out_Time) 
+{
+	ParticleID -= {ParameterName}_LastSpawnedPointID + 1;
+	Out_Location = {ParameterName}_IncomingLocationBuffer[ParticleID];
+	Out_AccumulatedImpulse = {ParameterName}_IncomingAccumulatedImpulseBuffer[ParticleID];
+	Out_Normal = {ParameterName}_IncomingNormalBuffer[ParticleID];
+	Out_Velocity1 = {ParameterName}_IncomingVelocity1Buffer[ParticleID];
+	Out_Velocity2 = {ParameterName}_IncomingVelocity2Buffer[ParticleID];
+	Out_AngularVelocity1 = {ParameterName}_IncomingAngularVelocity1Buffer[ParticleID];
+	Out_AngularVelocity2 = {ParameterName}_IncomingAngularVelocity2Buffer[ParticleID];
+	Out_Mass1 = {ParameterName}_IncomingMass1Buffer[ParticleID];
+	Out_Mass2 = {ParameterName}_IncomingMass2Buffer[ParticleID];
+	Out_Time = {ParameterName}_IncomingTimeBuffer[ParticleID];
+}
+
+void GetBreakingData_{ParameterName}(in int ParticleID, out float3 Out_Location, out float3 Out_Velocity, out float3 Out_AngularVelocity, out float Out_Mass, out float Out_Time) 
+{
+	ParticleID -= {ParameterName}_LastSpawnedPointID + 1;
+	Out_Location = {ParameterName}_IncomingLocationBuffer[ParticleID];
+	Out_Velocity = {ParameterName}_IncomingVelocity1Buffer[ParticleID];
+	Out_AngularVelocity = {ParameterName}_IncomingAngularVelocity1Buffer[ParticleID];
+	Out_Mass = {ParameterName}_IncomingMass1Buffer[ParticleID];
+	Out_Time = {ParameterName}_IncomingTimeBuffer[ParticleID];
+}
+
+void GetTrailingData_{ParameterName}(in int ParticleID, out float3 Out_Location, out float3 Out_Velocity, out float3 Out_AngularVelocity, out float Out_Mass, out float Out_Time)
+{
+	ParticleID -= {ParameterName}_LastSpawnedPointID + 1;
+	Out_Location = {ParameterName}_IncomingLocationBuffer[ParticleID];
+	Out_Velocity = {ParameterName}_IncomingVelocity1Buffer[ParticleID];
+	Out_AngularVelocity = {ParameterName}_IncomingAngularVelocity1Buffer[ParticleID];
+	Out_Mass = {ParameterName}_IncomingMass1Buffer[ParticleID];
+	Out_Time = {ParameterName}_IncomingTimeBuffer[ParticleID];
+}
diff --git a/Engine/Plugins/Experimental/ChaosNiagara/Shaders/NiagaraDataInterfaceGeometryCollection.ush b/Engine/Plugins/Experimental/ChaosNiagara/Shaders/NiagaraDataInterfaceGeometryCollection.ush
--- a/Engine/Plugins/Experimental/ChaosNiagara/Shaders/NiagaraDataInterfaceGeometryCollection.ush
+++ b/Engine/Plugins/Experimental/ChaosNiagara/Shaders/NiagaraDataInterfaceGeometryCollection.ush
@@ -4,8 +4,6 @@
 NiagaraDataInterfaceGeometryCollection.ush 
 =============================================================================*/
 
-#pragma once
-
 /* -----------------------------------------------------------------
  * static mesh constants and context
  * -----------------------------------------------------------------
@@ -17,81 +15,54 @@ NiagaraDataInterfaceGeometryCollection.ush
 #define SPHERE_INDEX 1
 #define CAPSULE_INDEX 2
 
-struct FDIGeometryCollectionContext
-{
-	float3 BoundsMin;
-	float3 BoundsMax;
-	uint NumPieces;	
-	Buffer<float4>		WorldTransformBuffer;
-	Buffer<float4>		PrevWorldTransformBuffer;
-	Buffer<float4>		WorldInverseTransformBuffer;
-	Buffer<float4>		PrevWorldInverseTransformBuffer;
-	Buffer<float4>		BoundsBuffer;
-}; 
-
-#define DIGEOMETRYCOLLECTION_DECLARE_CONSTANTS(NAME)\
-float3 BoundsMin_##NAME;\
-float3 BoundsMax_##NAME;\
-uint NumPieces_##NAME;\
-Buffer<float4>		WorldTransformBuffer_##NAME;\
-Buffer<float4>		PrevWorldTransformBuffer_##NAME;\
-Buffer<float4>		WorldInverseTransformBuffer_##NAME;\
-Buffer<float4>		PrevWorldInverseTransformBuffer_##NAME;\
-Buffer<float4>		BoundsBuffer_##NAME;\
-
-
-#define DIGEOMETRYCOLLECTION_MAKE_CONTEXT(NAME)\
-FDIGeometryCollectionContext DIContext; \
-DIContext.BoundsMin = BoundsMin_##NAME;\
-DIContext.BoundsMax = BoundsMax_##NAME;\
-DIContext.NumPieces = NumPieces_##NAME;\
-DIContext.WorldTransformBuffer =  WorldTransformBuffer_##NAME;\
-DIContext.PrevWorldTransformBuffer = PrevWorldTransformBuffer_##NAME;\
-DIContext.WorldInverseTransformBuffer = WorldInverseTransformBuffer_##NAME;\
-DIContext.PrevWorldInverseTransformBuffer = PrevWorldInverseTransformBuffer_##NAME;\
-DIContext.BoundsBuffer = BoundsBuffer_##NAME;\
-
-
+float3			{ParameterName}_BoundsMin;
+float3			{ParameterName}_BoundsMax;
+uint			{ParameterName}_NumPieces;
+Buffer<float4>	{ParameterName}_WorldTransformBuffer;
+Buffer<float4>	{ParameterName}_PrevWorldTransformBuffer;
+Buffer<float4>	{ParameterName}_WorldInverseTransformBuffer;
+Buffer<float4>	{ParameterName}_PrevWorldInverseTransformBuffer;
+Buffer<float4>	{ParameterName}_BoundsBuffer;
 
 /* -----------------------------------------------------------------
  * Position / Transform utils
  * -----------------------------------------------------------------
  */
 
-float3 DIGeometryCollection_GetCurrentPosition(in FDIGeometryCollectionContext DIContext, in int ElementIndex)
+float3 DIGeometryCollection_GetCurrentPosition(in int ElementIndex)
 {
 	const int BufferOffset = 3 * ElementIndex;
-    return float3(DIContext.WorldTransformBuffer[BufferOffset].w, DIContext.WorldTransformBuffer[BufferOffset+1].w, DIContext.WorldTransformBuffer[BufferOffset+2].w);
+    return float3({ParameterName}_WorldTransformBuffer[BufferOffset].w, {ParameterName}_WorldTransformBuffer[BufferOffset+1].w, {ParameterName}_WorldTransformBuffer[BufferOffset+2].w);
 }
 
-float3 DIGeometryCollection_GetPreviousPosition(in FDIGeometryCollectionContext DIContext, in int ElementIndex)
+float3 DIGeometryCollection_GetPreviousPosition(in int ElementIndex)
 {
 	const int BufferOffset = 3 * ElementIndex;
-    return float3(DIContext.WorldTransformBuffer[BufferOffset].w, DIContext.WorldTransformBuffer[BufferOffset+1].w, DIContext.WorldTransformBuffer[BufferOffset+2].w);
+    return float3({ParameterName}_WorldTransformBuffer[BufferOffset].w, {ParameterName}_WorldTransformBuffer[BufferOffset+1].w, {ParameterName}_WorldTransformBuffer[BufferOffset+2].w);
 }
 
-float3x4 DIGeometryCollection_GetCurrentTransform(in FDIGeometryCollectionContext DIContext, in int ElementIndex)
+float3x4 DIGeometryCollection_GetCurrentTransform(in int ElementIndex)
 {
 	const int BufferOffset = 3 * ElementIndex;
-    return float3x4(DIContext.WorldTransformBuffer[BufferOffset], DIContext.WorldTransformBuffer[BufferOffset+1], DIContext.WorldTransformBuffer[BufferOffset+2]);
+    return float3x4({ParameterName}_WorldTransformBuffer[BufferOffset], {ParameterName}_WorldTransformBuffer[BufferOffset+1], {ParameterName}_WorldTransformBuffer[BufferOffset+2]);
 }
 
-float3x4 DIGeometryCollection_GetInverseTransform(in FDIGeometryCollectionContext DIContext, in int ElementIndex)
+float3x4 DIGeometryCollection_GetInverseTransform(in int ElementIndex)
 {
 	const int BufferOffset = 3 * ElementIndex;
-    return float3x4(DIContext.WorldInverseTransformBuffer[BufferOffset], DIContext.WorldInverseTransformBuffer[BufferOffset+1], DIContext.WorldInverseTransformBuffer[BufferOffset+2]);
+    return float3x4({ParameterName}_WorldInverseTransformBuffer[BufferOffset], {ParameterName}_WorldInverseTransformBuffer[BufferOffset+1], {ParameterName}_WorldInverseTransformBuffer[BufferOffset+2]);
 }
 
-float3x4 DIGeometryCollection_GetPreviousTransform(in FDIGeometryCollectionContext DIContext, in int ElementIndex)
+float3x4 DIGeometryCollection_GetPreviousTransform(in int ElementIndex)
 {
 	const int BufferOffset = 3 * ElementIndex;
-    return float3x4(DIContext.PrevWorldTransformBuffer[BufferOffset], DIContext.PrevWorldTransformBuffer[BufferOffset+1], DIContext.PrevWorldTransformBuffer[BufferOffset+2]);
+    return float3x4({ParameterName}_PrevWorldTransformBuffer[BufferOffset], {ParameterName}_PrevWorldTransformBuffer[BufferOffset+1], {ParameterName}_PrevWorldTransformBuffer[BufferOffset+2]);
 }
 
-float3x4 DIGeometryCollection_GetPreviousInverse(in FDIGeometryCollectionContext DIContext, in int ElementIndex)
+float3x4 DIGeometryCollection_GetPreviousInverse(in int ElementIndex)
 {
 	const int BufferOffset = 3 * ElementIndex;
-    return float3x4(DIContext.PrevWorldInverseTransformBuffer[BufferOffset], DIContext.PrevWorldInverseTransformBuffer[BufferOffset+1], DIContext.PrevWorldInverseTransformBuffer[BufferOffset+2]);
+    return float3x4({ParameterName}_PrevWorldInverseTransformBuffer[BufferOffset], {ParameterName}_PrevWorldInverseTransformBuffer[BufferOffset+1], {ParameterName}_PrevWorldInverseTransformBuffer[BufferOffset+2]);
 }
 
 
@@ -183,16 +154,16 @@ float DIGeometryCollection_GetBoxProjection(in float3 LocalPosition, in float3 B
  * -----------------------------------------------------------------
  */
 
-float3 DIGeometryCollection_GetLocalPosition(in FDIGeometryCollectionContext DIContext, in float3 WorldPosition, in int ElementIndex, in float TimeFraction )
+float3 DIGeometryCollection_GetLocalPosition(in float3 WorldPosition, in int ElementIndex, in float TimeFraction )
 {
-	const float3 CurrentLocal = mul(DIGeometryCollection_GetInverseTransform(DIContext,ElementIndex), float4(WorldPosition,1.0)).xyz;
-	const float3 PreviousLocal = mul(DIGeometryCollection_GetPreviousInverse(DIContext,ElementIndex), float4(WorldPosition,1.0)).xyz;
+	const float3 CurrentLocal = mul(DIGeometryCollection_GetInverseTransform(ElementIndex), float4(WorldPosition,1.0)).xyz;
+	const float3 PreviousLocal = mul(DIGeometryCollection_GetPreviousInverse(ElementIndex), float4(WorldPosition,1.0)).xyz;
 
 	return PreviousLocal + TimeFraction * (CurrentLocal-PreviousLocal);
 }
 
 // Given a world space position (WorldPosition) compute the static mesh closest point (position,normal,velocity)
-void DIGeometryCollection_GetClosestPointNoNormal(in FDIGeometryCollectionContext DIContext, in float3 WorldPosition, in float DeltaTime, in float TimeFraction, out float OutClosestDistance, out float3 OutClosestPosition, out float3 OutClosestVelocity)
+void DIGeometryCollection_GetClosestPointNoNormal(in float3 WorldPosition, in float DeltaTime, in float TimeFraction, out float OutClosestDistance, out float3 OutClosestPosition, out float3 OutClosestVelocity)
 {	
 	float3 CollisionPosition = float3(0,0,0);
 	float3 CollisionNormal = float3(0,0,0);
@@ -203,31 +174,44 @@ void DIGeometryCollection_GetClosestPointNoNormal(in FDIGeometryCollectionContex
 	
 	
 	[branch]
-	if (WorldPosition.x < DIContext.BoundsMin.x || 
-		WorldPosition.y < DIContext.BoundsMin.y || 
-		WorldPosition.z < DIContext.BoundsMin.z || 
-		WorldPosition.x > DIContext.BoundsMax.x || 
-		WorldPosition.y > DIContext.BoundsMax.y || 
-		WorldPosition.z > DIContext.BoundsMax.z)
+	if (WorldPosition.x < {ParameterName}_BoundsMin.x || 
+		WorldPosition.y < {ParameterName}_BoundsMin.y || 
+		WorldPosition.z < {ParameterName}_BoundsMin.z || 
+		WorldPosition.x > {ParameterName}_BoundsMax.x || 
+		WorldPosition.y > {ParameterName}_BoundsMax.y || 
+		WorldPosition.z > {ParameterName}_BoundsMax.z)
 	{
 		return;
 	}
 	
 
 	int ClosestTransform = -1;
-	for (int CurrTransform = 0; CurrTransform < DIContext.NumPieces; ++CurrTransform)
+	for (int CurrTransform = 0; CurrTransform < {ParameterName}_NumPieces; ++CurrTransform)
 	{
-		float3 LocalPosition = DIGeometryCollection_GetLocalPosition(DIContext,WorldPosition,CurrTransform,TimeFraction);
+		float3 LocalPosition = DIGeometryCollection_GetLocalPosition(WorldPosition,CurrTransform,TimeFraction);
 		
-		DIGeometryCollection_GetBoxProjection(LocalPosition, DIContext.BoundsBuffer[CurrTransform].xyz, CurrTransform,
+		DIGeometryCollection_GetBoxProjection(LocalPosition, {ParameterName}_BoundsBuffer[CurrTransform].xyz, CurrTransform,
 				CollisionPosition, CollisionNormal, ClosestTransform, OutClosestDistance);
 	}
 
 	if (ClosestTransform != -1)
 	{
-		const float3 PreviousPosition = mul(DIGeometryCollection_GetPreviousTransform(DIContext,ClosestTransform), float4(CollisionPosition,1.0)).xyz;
-		const float3 CurrentPosition = mul(DIGeometryCollection_GetCurrentTransform(DIContext,ClosestTransform), float4(CollisionPosition,1.0)).xyz;
+		const float3 PreviousPosition = mul(DIGeometryCollection_GetPreviousTransform(ClosestTransform), float4(CollisionPosition,1.0)).xyz;
+		const float3 CurrentPosition = mul(DIGeometryCollection_GetCurrentTransform(ClosestTransform), float4(CollisionPosition,1.0)).xyz;
 		OutClosestVelocity = ( CurrentPosition - PreviousPosition ) / DeltaTime;
 		OutClosestPosition = PreviousPosition + TimeFraction * (CurrentPosition-PreviousPosition);		
 	}
-}
\ No newline at end of file
+}
+
+
+void GetClosestPointNoNormal_{ParameterName}(in float3 WorldPosition, in float DeltaTime, in float TimeFraction, out float ClosestDistance, out float3 OutClosestPosition, out float3 OutClosestVelocity)
+{
+	DIGeometryCollection_GetClosestPointNoNormal(WorldPosition,DeltaTime,TimeFraction, ClosestDistance,
+		OutClosestPosition,OutClosestVelocity);
+}
+
+#undef SMALL_NUMBER
+#undef MAX_DISTANCE
+#undef BOX_INDEX
+#undef SPHERE_INDEX
+#undef CAPSULE_INDEX
diff --git a/Engine/Plugins/Experimental/ChaosNiagara/Source/ChaosNiagara/Classes/NiagaraDataInterfaceChaosDestruction.h b/Engine/Plugins/Experimental/ChaosNiagara/Source/ChaosNiagara/Classes/NiagaraDataInterfaceChaosDestruction.h
--- a/Engine/Plugins/Experimental/ChaosNiagara/Source/ChaosNiagara/Classes/NiagaraDataInterfaceChaosDestruction.h
+++ b/Engine/Plugins/Experimental/ChaosNiagara/Source/ChaosNiagara/Classes/NiagaraDataInterfaceChaosDestruction.h
@@ -135,21 +135,21 @@ struct FNiagaraDIChaosDestruction_GPUData
 	TArray<int32> SurfaceTypeArray;
 	TArray<FLinearColor> ColorArray;
 
-	TArray<FVector> IncomingLocationArray;				// Collision, Breaking, Trailing
-	TArray<FVector> IncomingAccumulatedImpulseArray;	// Collision
-	TArray<FVector> IncomingNormalArray;				// Collision
-	TArray<FVector> IncomingVelocity1Array;			// Collision, Breaking, Trailing
-	TArray<FVector> IncomingVelocity2Array;			// Collision
-	TArray<FVector> IncomingAngularVelocity1Array;		// Collision, Breaking, Trailing
-	TArray<FVector> IncomingAngularVelocity2Array;		// Collision
+	TArray<FVector3f> IncomingLocationArray;				// Collision, Breaking, Trailing
+	TArray<FVector3f> IncomingAccumulatedImpulseArray;	// Collision
+	TArray<FVector3f> IncomingNormalArray;				// Collision
+	TArray<FVector3f> IncomingVelocity1Array;			// Collision, Breaking, Trailing
+	TArray<FVector3f> IncomingVelocity2Array;			// Collision
+	TArray<FVector3f> IncomingAngularVelocity1Array;		// Collision, Breaking, Trailing
+	TArray<FVector3f> IncomingAngularVelocity2Array;		// Collision
 	TArray<float> IncomingMass1Array;					// Collision, Breaking, Trailing
 	TArray<float> IncomingMass2Array;					// Collision
 	TArray<float> IncomingTimeArray;					// Collision, Breaking, Trailing
 
-	TArray<FVector> TransformTranslationArray;
-	TArray<FQuat> TransformRotationArray;
-	TArray<FVector> TransformScaleArray;
-	TArray<FVector> BoundsArray;
+	TArray<FVector3f> TransformTranslationArray;
+	TArray<FQuat4f> TransformRotationArray;
+	TArray<FVector3f> TransformScaleArray;
+	TArray<FVector3f> BoundsArray;
 
 	void ResetAll()
 	{
@@ -182,67 +182,65 @@ struct FNiagaraDIChaosDestruction_GPUData
 		BoundsArray.Reset();
 	}
 
-	FDynamicReadBuffer GPUPositionBuffer;
-	FDynamicReadBuffer GPUVelocityBuffer;
-	FDynamicReadBuffer GPUExtentMinBuffer;
-	FDynamicReadBuffer GPUExtentMaxBuffer;
-	FDynamicReadBuffer GPUVolumeBuffer;
-	FDynamicReadBuffer GPUSolverIDBuffer;
-	FDynamicReadBuffer GPUDensityBuffer;
-	FDynamicReadBuffer GPUFrictionBuffer;
-	FDynamicReadBuffer GPURestitutionBuffer;
-	FDynamicReadBuffer GPUSurfaceTypeBuffer;
-	FDynamicReadBuffer GPUColorBuffer;
-
-	FDynamicReadBuffer GPUIncomingLocationBuffer;
-	FDynamicReadBuffer GPUIncomingAccumulatedImpulseBuffer;
-	FDynamicReadBuffer GPUIncomingNormalBuffer;
-	FDynamicReadBuffer GPUIncomingVelocity1Buffer;
-	FDynamicReadBuffer GPUIncomingVelocity2Buffer;
-	FDynamicReadBuffer GPUIncomingAngularVelocity1Buffer;
-	FDynamicReadBuffer GPUIncomingAngularVelocity2Buffer;
-	FDynamicReadBuffer GPUIncomingMass1Buffer;
-	FDynamicReadBuffer GPUIncomingMass2Buffer;
-	FDynamicReadBuffer GPUIncomingTimeBuffer;
-
-	FDynamicReadBuffer GPUTransformTranslationBuffer;
-	FDynamicReadBuffer GPUTransformRotationBuffer;
-	FDynamicReadBuffer GPUTransformScaleBuffer;
-	FDynamicReadBuffer GPUBoundsBuffer;
-
 	float SolverTime;
 	int32 LastSpawnedPointID;
+
+	FRDGBufferSRVRef RDGPositionSRV = nullptr;
+	FRDGBufferSRVRef RDGVelocitySRV = nullptr;
+	FRDGBufferSRVRef RDGExtentMinBufferSRV = nullptr;
+	FRDGBufferSRVRef RDGExtentMaxBufferSRV = nullptr;
+	FRDGBufferSRVRef RDGVolumeBufferSRV = nullptr;
+	FRDGBufferSRVRef RDGSolverIDBufferSRV = nullptr;
+	FRDGBufferSRVRef RDGDensityBufferSRV = nullptr;
+	FRDGBufferSRVRef RDGFrictionBufferSRV = nullptr;
+	FRDGBufferSRVRef RDGRestitutionBufferSRV = nullptr;
+	FRDGBufferSRVRef RDGSurfaceTypeBufferSRV = nullptr;
+	FRDGBufferSRVRef RDGColorBufferSRV = nullptr;
+	FRDGBufferSRVRef RDGIncomingLocationBufferSRV = nullptr;
+	FRDGBufferSRVRef RDGIncomingAccumulatedImpulseBufferSRV = nullptr;
+	FRDGBufferSRVRef RDGIncomingNormalBufferSRV = nullptr;
+	FRDGBufferSRVRef RDGIncomingVelocity1BufferSRV = nullptr;
+	FRDGBufferSRVRef RDGIncomingVelocity2BufferSRV = nullptr;
+	FRDGBufferSRVRef RDGIncomingAngularVelocity1BufferSRV = nullptr;
+	FRDGBufferSRVRef RDGIncomingAngularVelocity2BufferSRV = nullptr;
+	FRDGBufferSRVRef RDGIncomingMass1BufferSRV = nullptr;
+	FRDGBufferSRVRef RDGIncomingMass2BufferSRV = nullptr;
+	FRDGBufferSRVRef RDGIncomingTimeBufferSRV = nullptr;
+	FRDGBufferSRVRef RDGTransformTranslationBufferSRV = nullptr;
+	FRDGBufferSRVRef RDGTransformRotationBufferSRV = nullptr;
+	FRDGBufferSRVRef RDGTransformScaleBufferSRV = nullptr;
+	FRDGBufferSRVRef RDGBoundsBufferSRV = nullptr;
 };
 
 struct FNiagaraDIChaosDestruction_InstanceDataToPassToRT
 {
-	TArray<FVector3f>* PositionArray;
-	TArray<FVector3f>* VelocityArray;
-	TArray<float>* ExtentMinArray;
-	TArray<float>* ExtentMaxArray;
-	TArray<float>* VolumeArray;
-	TArray<int32>* SolverIDArray;
-	TArray<float>* DensityArray;
-	TArray<float>* FrictionArray;
-	TArray<float>* RestitutionArray;
-	TArray<int32>* SurfaceTypeArray;
-	TArray<FLinearColor>* ColorArray;
-
-	TArray<FVector>* IncomingLocationArray;				// Collision, Breaking, Trailing
-	TArray<FVector>* IncomingAccumulatedImpulseArray;	// Collision
-	TArray<FVector>* IncomingNormalArray;				// Collision
-	TArray<FVector>* IncomingVelocity1Array;			// Collision, Breaking, Trailing
-	TArray<FVector>* IncomingVelocity2Array;			// Collision
-	TArray<FVector>* IncomingAngularVelocity1Array;		// Collision, Breaking, Trailing
-	TArray<FVector>* IncomingAngularVelocity2Array;		// Collision
-	TArray<float>* IncomingMass1Array;					// Collision, Breaking, Trailing
-	TArray<float>* IncomingMass2Array;					// Collision
-	TArray<float>* IncomingTimeArray;					// Collision, Breaking, Trailing
+	TArray<FVector3f> PositionArray;
+	TArray<FVector3f> VelocityArray;
+	TArray<float> ExtentMinArray;
+	TArray<float> ExtentMaxArray;
+	TArray<float> VolumeArray;
+	TArray<int32> SolverIDArray;
+	TArray<float> DensityArray;
+	TArray<float> FrictionArray;
+	TArray<float> RestitutionArray;
+	TArray<int32> SurfaceTypeArray;
+	TArray<FLinearColor> ColorArray;
+
+	TArray<FVector3f> IncomingLocationArray;				// Collision, Breaking, Trailing
+	TArray<FVector3f> IncomingAccumulatedImpulseArray;	// Collision
+	TArray<FVector3f> IncomingNormalArray;				// Collision
+	TArray<FVector3f> IncomingVelocity1Array;			// Collision, Breaking, Trailing
+	TArray<FVector3f> IncomingVelocity2Array;			// Collision
+	TArray<FVector3f> IncomingAngularVelocity1Array;		// Collision, Breaking, Trailing
+	TArray<FVector3f> IncomingAngularVelocity2Array;		// Collision
+	TArray<float> IncomingMass1Array;					// Collision, Breaking, Trailing
+	TArray<float> IncomingMass2Array;					// Collision
+	TArray<float> IncomingTimeArray;					// Collision, Breaking, Trailing
 	
-	TArray<FVector>* TransformTranslationArray; // Breaking
-	TArray<FQuat>* TransformRotationArray; // Breaking
-	TArray<FVector>* TransformScaleArray; // Breaking
-	TArray<FVector>* BoundsArray; // Breaking
+	TArray<FVector3f> TransformTranslationArray; // Breaking
+	TArray<FQuat4f> TransformRotationArray; // Breaking
+	TArray<FVector3f> TransformScaleArray; // Breaking
+	TArray<FVector3f> BoundsArray; // Breaking
 
 	float SolverTime;
 	int32 LastSpawnedPointID;
@@ -352,9 +350,6 @@ class CHAOSNIAGARA_API UNiagaraDataInterfaceChaosDestruction : public UNiagaraDa
 	GENERATED_UCLASS_BODY()
 
 public:
-
-	DECLARE_NIAGARA_DI_PARAMETER();
-
 	/* Chaos Solver */
 	UPROPERTY(EditAnywhere, Category = "Solver", meta = (DisplayName = "Chaos Solver"))
 	TSet<TObjectPtr<AChaosSolverActor>> ChaosSolverActorSet;
@@ -563,9 +558,13 @@ public:
 	//----------------------------------------------------------------------------
 	// GPU sim functionality
 #if WITH_EDITORONLY_DATA
+	virtual bool AppendCompileHash(FNiagaraCompileHashVisitor* InVisitor) const override;
 	virtual void GetParameterDefinitionHLSL(const FNiagaraDataInterfaceGPUParamInfo& ParamInfo, FString& OutHLSL) override;
 	virtual bool GetFunctionHLSL(const FNiagaraDataInterfaceGPUParamInfo& ParamInfo, const FNiagaraDataInterfaceGeneratedFunction& FunctionInfo, int FunctionInstanceIndex, FString& OutHLSL) override;
 #endif
+	virtual bool UseLegacyShaderBindings() const override { return false; }
+	virtual void BuildShaderParameters(FNiagaraShaderParametersBuilder& ShaderParametersBuilder) const override;
+	virtual void SetShaderParameters(const FNiagaraDataInterfaceSetShaderParametersContext& Context) const override;
 
 	//----------------------------------------------------------------------------
 	// EXPOSED FUNCTIONS
@@ -769,10 +768,9 @@ protected:
 struct FNiagaraDataInterfaceProxyChaosDestruction : public FNiagaraDataInterfaceProxy
 {
 	virtual void ConsumePerInstanceDataFromGameThread(void* PerInstanceData, const FNiagaraSystemInstanceID& Instance) override;
-	virtual int32 PerInstanceDataPassedToRenderThreadSize() const override
-	{
-		return sizeof(FNiagaraDIChaosDestruction_InstanceDataToPassToRT);
-	}
+	virtual int32 PerInstanceDataPassedToRenderThreadSize() const override { return sizeof(FNiagaraDIChaosDestruction_InstanceDataToPassToRT); }
+
+	virtual void PostSimulate(const FNDIGpuComputePostSimulateContext& Context) override;
 
 	void CreatePerInstanceData(const FNiagaraSystemInstanceID& SystemInstance);
 
@@ -782,4 +780,4 @@ struct FNiagaraDataInterfaceProxyChaosDestruction : public FNiagaraDataInterface
 	int32 LastSpawnedPointID;
 
 	TMap<FNiagaraSystemInstanceID, FNiagaraDIChaosDestruction_GPUData> SystemsToGPUInstanceData;
-};
\ No newline at end of file
+};
diff --git a/Engine/Plugins/Experimental/ChaosNiagara/Source/ChaosNiagara/Private/NiagaraDataInterfaceChaosDestruction.cpp b/Engine/Plugins/Experimental/ChaosNiagara/Source/ChaosNiagara/Private/NiagaraDataInterfaceChaosDestruction.cpp
--- a/Engine/Plugins/Experimental/ChaosNiagara/Source/ChaosNiagara/Private/NiagaraDataInterfaceChaosDestruction.cpp
+++ b/Engine/Plugins/Experimental/ChaosNiagara/Source/ChaosNiagara/Private/NiagaraDataInterfaceChaosDestruction.cpp
@@ -3,7 +3,9 @@
 #include "NiagaraDataInterfaceChaosDestruction.h"
 #include "NiagaraTypes.h"
 #include "Misc/FileHelper.h"
+#include "NiagaraGpuComputeDispatchInterface.h"
 #include "NiagaraShader.h"
+#include "NiagaraShaderParametersBuilder.h"
 #include "ShaderParameterUtils.h"
 #include "PhysicsSolver.h"
 #include "Niagara/Private/NiagaraStats.h"
@@ -47,27 +49,159 @@ DECLARE_MEMORY_STAT(TEXT("AllBreakingsIndicesByPhysicsProxy"), STAT_AllBreakings
 DECLARE_MEMORY_STAT(TEXT("AllTrailingsData"), STAT_AllTrailingsDataMemory, STATGROUP_ChaosNiagara);
 DECLARE_MEMORY_STAT(TEXT("AllTrailingsIndicesByPhysicsProxy"), STAT_AllTrailingsIndicesByPhysicsProxyMemory, STATGROUP_ChaosNiagara);
 
-// Name of all the functions available in the data interface
-static const FName GetPositionName("GetPosition");
-static const FName GetNormalName("GetNormal");
-static const FName GetVelocityName("GetVelocity");
-static const FName GetAngularVelocityName("GetAngularVelocity");
-static const FName GetExtentMinName("GetExtentMin");
-static const FName GetExtentMaxName("GetExtentMax");
-static const FName GetVolumeName("GetVolume");
-static const FName GetParticleIdsToSpawnAtTimeName("GetParticleIdsToSpawnAtTime");
-static const FName GetPointTypeName("GetPointType");
-static const FName GetColorName("GetColor");
-static const FName GetSolverTimeName("GetSolverTime");
-static const FName GetDensityName("GetDensity");
-static const FName GetFrictionName("GetFriction");
-static const FName GetRestitutionName("GetRestitution");
-static const FName GetSurfaceTypeName("GetSurfaceType");
-static const FName GetTransformName("GetTransform");
-static const FName GetSizeName("GetSize");
-static const FName GetCollisionDataName("GetCollisionData");
-static const FName GetBreakingDataName("GetBreakingData");
-static const FName GetTrailingDataName("GetTrailingData");
+namespace NDIChaosDestructionLocal
+{
+	static const TCHAR* TemplateShaderFilePath = TEXT("/Plugin/Experimental/ChaosNiagara/NiagaraDataInterfaceChaosDestruction.ush");
+
+	BEGIN_SHADER_PARAMETER_STRUCT(FShaderParameters, )
+		SHADER_PARAMETER_RDG_BUFFER_SRV(Buffer<float3>,	PositionBuffer)
+		SHADER_PARAMETER_RDG_BUFFER_SRV(Buffer<float3>,	VelocityBuffer)
+		SHADER_PARAMETER_RDG_BUFFER_SRV(Buffer<float>,	ExtentMinBuffer)
+		SHADER_PARAMETER_RDG_BUFFER_SRV(Buffer<float>,	ExtentMaxBuffer)
+		SHADER_PARAMETER_RDG_BUFFER_SRV(Buffer<float>,	VolumeBuffer)
+		SHADER_PARAMETER_RDG_BUFFER_SRV(Buffer<int>,	SolverIDBuffer)
+		SHADER_PARAMETER_RDG_BUFFER_SRV(Buffer<float>,	DensityBuffer)
+		SHADER_PARAMETER_RDG_BUFFER_SRV(Buffer<float>,	FrictionBuffer)
+		SHADER_PARAMETER_RDG_BUFFER_SRV(Buffer<float>,	RestitutionBuffer)
+		SHADER_PARAMETER_RDG_BUFFER_SRV(Buffer<int>,	SurfaceTypeBuffer)
+		SHADER_PARAMETER_RDG_BUFFER_SRV(Buffer<float4>,	ColorBuffer)
+		
+		SHADER_PARAMETER_RDG_BUFFER_SRV(Buffer<float3>,	IncomingLocationBuffer)
+		SHADER_PARAMETER_RDG_BUFFER_SRV(Buffer<float3>, IncomingAccumulatedImpulseBuffer)
+		SHADER_PARAMETER_RDG_BUFFER_SRV(Buffer<float3>, IncomingNormalBuffer)
+		SHADER_PARAMETER_RDG_BUFFER_SRV(Buffer<float3>, IncomingVelocity1Buffer)
+		SHADER_PARAMETER_RDG_BUFFER_SRV(Buffer<float3>, IncomingVelocity2Buffer)
+		SHADER_PARAMETER_RDG_BUFFER_SRV(Buffer<float3>, IncomingAngularVelocity1Buffer)
+		SHADER_PARAMETER_RDG_BUFFER_SRV(Buffer<float3>, IncomingAngularVelocity2Buffer)
+		SHADER_PARAMETER_RDG_BUFFER_SRV(Buffer<float>, IncomingMass1Buffer)
+		SHADER_PARAMETER_RDG_BUFFER_SRV(Buffer<float>, IncomingMass2Buffer)
+		SHADER_PARAMETER_RDG_BUFFER_SRV(Buffer<float>, IncomingTimeBuffer)
+
+		SHADER_PARAMETER_RDG_BUFFER_SRV(Buffer<float3>, TransformTranslationBuffer)
+		SHADER_PARAMETER_RDG_BUFFER_SRV(Buffer<float4>, TransformRotationBuffer)
+		SHADER_PARAMETER_RDG_BUFFER_SRV(Buffer<float3>, TransformScaleBuffer)
+		SHADER_PARAMETER_RDG_BUFFER_SRV(Buffer<float3>, BoundsBuffer)
+
+		SHADER_PARAMETER(int,	LastSpawnedPointID)
+		SHADER_PARAMETER(float,	SolverTime)
+	END_SHADER_PARAMETER_STRUCT()
+
+	// Name of all the functions available in the data interface
+	static const FName GetPositionName("GetPosition");
+	static const FName GetNormalName("GetNormal");
+	static const FName GetVelocityName("GetVelocity");
+	static const FName GetAngularVelocityName("GetAngularVelocity");
+	static const FName GetExtentMinName("GetExtentMin");
+	static const FName GetExtentMaxName("GetExtentMax");
+	static const FName GetVolumeName("GetVolume");
+	static const FName GetParticleIdsToSpawnAtTimeName("GetParticleIdsToSpawnAtTime");
+	static const FName GetPointTypeName("GetPointType");
+	static const FName GetColorName("GetColor");
+	static const FName GetSolverTimeName("GetSolverTime");
+	static const FName GetDensityName("GetDensity");
+	static const FName GetFrictionName("GetFriction");
+	static const FName GetRestitutionName("GetRestitution");
+	static const FName GetSurfaceTypeName("GetSurfaceType");
+	static const FName GetTransformName("GetTransform");
+	static const FName GetSizeName("GetSize");
+	static const FName GetCollisionDataName("GetCollisionData");
+	static const FName GetBreakingDataName("GetBreakingData");
+	static const FName GetTrailingDataName("GetTrailingData");
+
+	TArray<FVector3f> CopyVectorArray(TArrayView<FVector> DoubleArray)
+	{
+		TArray<FVector3f> FloatArray;
+		FloatArray.AddUninitialized(DoubleArray.Num());
+		for ( int i=0; i < DoubleArray.Num(); ++i )
+		{
+			FloatArray[i] = FVector3f(DoubleArray[i]);	//LWC_TODO: Should we convert the space?
+		}
+		return FloatArray;
+	}
+
+	TArray<FQuat4f> CopyQuatArray(TArrayView<FQuat> DoubleArray)
+	{
+		TArray<FQuat4f> FloatArray;
+		FloatArray.AddUninitialized(DoubleArray.Num());
+		for (int i = 0; i < DoubleArray.Num(); ++i)
+		{
+			FloatArray[i] = FQuat4f(DoubleArray[i]);	//LWC_TODO: Should we convert the space?
+		}
+		return FloatArray;
+	}
+
+	template<typename T>
+	void QueueBufferUpload(const FNiagaraDataInterfaceSetShaderParametersContext& Context, FRDGBufferSRVRef& ParameterBinding, FRDGBufferSRVRef& ExistingSRV, TConstArrayView<T> ArrayData, const T& DefaultValue, const EPixelFormat PixelFormat, const TCHAR* BufferName)
+	{
+		// Buffer Used?
+		if (Context.IsResourceBound(&ParameterBinding) == false)
+		{
+			return;
+		}
+
+		// Buffer needs uploading this frame?
+		if (ExistingSRV == nullptr)
+		{
+			if (ArrayData.Num() == 0)
+			{
+				ArrayData = MakeArrayView<const T>(&DefaultValue, 1);
+			}
+
+			check(GPixelFormats[PixelFormat].BlockBytes == sizeof(T));
+
+			FRDGBuilder& GraphBuilder = Context.GetGraphBuilder();
+			FRDGBufferRef CreatedBuffer = GraphBuilder.CreateBuffer(FRDGBufferDesc::CreateBufferDesc(sizeof(T), ArrayData.Num()), BufferName);
+			GraphBuilder.QueueBufferUpload(CreatedBuffer, ArrayData.GetData(), ArrayData.GetTypeSize() * ArrayData.Num());
+
+			ExistingSRV = GraphBuilder.CreateSRV(CreatedBuffer, PixelFormat);
+		}
+
+		// Done assign the result
+		ParameterBinding = ExistingSRV;
+	}
+
+	template<>
+	void QueueBufferUpload<FVector3f>(const FNiagaraDataInterfaceSetShaderParametersContext& Context, FRDGBufferSRVRef& ParameterBinding, FRDGBufferSRVRef& ExistingSRV, TConstArrayView<FVector3f> ArrayData, const FVector3f& DefaultValue, const EPixelFormat PixelFormat, const TCHAR* BufferName)
+	{
+		// Buffer Used?
+		if (Context.IsResourceBound(&ParameterBinding) == false)
+		{
+			return;
+		}
+
+		// Buffer needs uploading this frame?
+		if (ExistingSRV == nullptr)
+		{
+			if (ArrayData.Num() == 0)
+			{
+				ArrayData = MakeArrayView<const FVector3f>(&DefaultValue, 1);
+			}
+
+			check(PixelFormat == PF_A32B32G32R32F);
+
+			FRDGBuilder& GraphBuilder = Context.GetGraphBuilder();
+			FRDGBufferRef CreatedBuffer = GraphBuilder.CreateBuffer(FRDGBufferDesc::CreateBufferDesc(sizeof(FVector4f), ArrayData.Num()), BufferName);
+			FVector4f* BufferData = GraphBuilder.AllocPODArray<FVector4f>(ArrayData.Num());
+			for (int i = 0; i < ArrayData.Num(); ++i)
+			{
+				BufferData[i] = FVector4f(ArrayData[i]);
+			}
+			GraphBuilder.QueueBufferUpload(CreatedBuffer, BufferData, sizeof(FVector4f) * ArrayData.Num(), ERDGInitialDataFlags::NoCopy);
+			ExistingSRV = GraphBuilder.CreateSRV(CreatedBuffer, PixelFormat);
+		}
+
+		// Done assign the result
+		ParameterBinding = ExistingSRV;
+	}
+
+	void GetEmptyBuffer(const FNiagaraDataInterfaceSetShaderParametersContext& Context, FRDGBufferSRVRef& ParameterBinding, const EPixelFormat PixelFormat)
+	{
+		if (Context.IsResourceBound(&ParameterBinding) == true)
+		{
+			ParameterBinding = Context.GetComputeDispatchInterface().GetEmptyBufferSRV(Context.GetGraphBuilder(), PixelFormat);
+		}
+	}
+}
 
 UNiagaraDataInterfaceChaosDestruction::UNiagaraDataInterfaceChaosDestruction(FObjectInitializer const& ObjectInitializer)
 	: Super(ObjectInitializer)
@@ -2405,6 +2539,7 @@ bool UNiagaraDataInterfaceChaosDestruction::PerInstanceTick(void* PerInstanceDat
 // Returns the signature of all the functions available in the data interface
 void UNiagaraDataInterfaceChaosDestruction::GetFunctions(TArray<FNiagaraFunctionSignature>& OutFunctions)
 {
+	using namespace NDIChaosDestructionLocal;
 	{
 		// GetPosition
 		FNiagaraFunctionSignature Sig;
@@ -2749,6 +2884,8 @@ DEFINE_NDI_FUNC_BINDER(UNiagaraDataInterfaceChaosDestruction, GetTrailingData);
 
 void UNiagaraDataInterfaceChaosDestruction::GetVMExternalFunction(const FVMExternalFunctionBindingInfo& BindingInfo, void* InstanceData, FVMExternalFunction &OutFunc)
 {
+	using namespace NDIChaosDestructionLocal;
+
 	if (BindingInfo.Name == GetPositionName && BindingInfo.GetNumInputs() == 2 && BindingInfo.GetNumOutputs() == 3)
 	{
 		TNDIParamBinder<1, int32, NDI_FUNC_BINDER(UNiagaraDataInterfaceChaosDestruction, GetPosition)>::Bind(this, BindingInfo, InstanceData, OutFunc);
@@ -3211,555 +3348,132 @@ void UNiagaraDataInterfaceChaosDestruction::GetTrailingData(FVectorVMExternalFun
 // GPU sim functionality
 //
 #if WITH_EDITORONLY_DATA
-void UNiagaraDataInterfaceChaosDestruction::GetParameterDefinitionHLSL(const FNiagaraDataInterfaceGPUParamInfo& ParamInfo, FString& OutHLSL)
+bool UNiagaraDataInterfaceChaosDestruction::AppendCompileHash(FNiagaraCompileHashVisitor* InVisitor) const
 {
+	bool bSuccess = Super::AppendCompileHash(InVisitor);
+	bSuccess &= InVisitor->UpdateString(TEXT("UNiagaraDataInterfaceChaosDestructionSource"), GetShaderFileHash(NDIChaosDestructionLocal::TemplateShaderFilePath, EShaderPlatform::SP_PCD3D_SM5).ToString());
+	bSuccess &= InVisitor->UpdateShaderParameters<NDIChaosDestructionLocal::FShaderParameters>();
+	return bSuccess;
+}
 
-	// This will get indented in the generated HLSL, which won't look good. 
-	// On the other hand, it makes it really nice and readable here.
-	static const TCHAR *FormatDeclarations = TEXT(R"(
-		Buffer<float3> PositionBuffer_{Symbol};
-		Buffer<float3> VelocityBuffer_{Symbol};
-		Buffer<float>  ExtentMinBuffer_{Symbol};
-		Buffer<float>  ExtentMaxBuffer_{Symbol};
-		Buffer<float>  VolumeBuffer_{Symbol};
-		Buffer<int>    SolverIDBuffer_{Symbol}; // NOTE(mv): Not used?
-		Buffer<float>  DensityBuffer_{Symbol};
-		Buffer<float>  FrictionBuffer_{Symbol};
-		Buffer<float>  RestitutionBuffer_{Symbol};
-		Buffer<int>    SurfaceTypeBuffer_{Symbol};
-		Buffer<float4> ColorBuffer_{Symbol};
-		
-		Buffer<float3> IncomingLocationBuffer_{Symbol};
-		Buffer<float3> IncomingAccumulatedImpulseBuffer_{Symbol};
-		Buffer<float3> IncomingNormalBuffer_{Symbol};
-		Buffer<float3> IncomingVelocity1Buffer_{Symbol};
-		Buffer<float3> IncomingVelocity2Buffer_{Symbol};
-		Buffer<float3> IncomingAngularVelocity1Buffer_{Symbol};
-		Buffer<float3> IncomingAngularVelocity2Buffer_{Symbol};
-		Buffer<float>  IncomingMass1Buffer_{Symbol};
-		Buffer<float>  IncomingMass2Buffer_{Symbol};
-		Buffer<float>  IncomingTimeBuffer_{Symbol};
-
-		Buffer<float3>  TransformTranslationBuffer_{Symbol};
-		Buffer<float4>  TransformRotationBuffer_{Symbol};
-		Buffer<float3>  TransformScaleBuffer_{Symbol};
-		Buffer<float3>  BoundsBuffer_{Symbol};
-
-		// NOTE(mv): Not implemented in the CPU-side functionality yet. 
-		//           Returns 0 in GetPointType instead.
-		//           
-		// Buffer<int> PointTypeBuffer_{Symbol};
-
-		int   LastSpawnedPointID_{Symbol};
-
-		float SolverTime_{Symbol};
-	)");
-
-	TMap<FString, FStringFormatArg> ArgsDeclarations = {
-		{ TEXT("Symbol"), ParamInfo.DataInterfaceHLSLSymbol },
-	};
-
-	OutHLSL += FString::Format(FormatDeclarations, ArgsDeclarations);
+void UNiagaraDataInterfaceChaosDestruction::GetParameterDefinitionHLSL(const FNiagaraDataInterfaceGPUParamInfo& ParamInfo, FString& OutHLSL)
+{
+	const TMap<FString, FStringFormatArg> TemplateArgs = { {TEXT("ParameterName"),	ParamInfo.DataInterfaceHLSLSymbol}, };
 
-	/*
-	*/
+	FString TemplateFile;
+	LoadShaderSourceFile(NDIChaosDestructionLocal::TemplateShaderFilePath, EShaderPlatform::SP_PCD3D_SM5, &TemplateFile, nullptr);
+	OutHLSL += FString::Format(*TemplateFile, TemplateArgs);
 }
 
 bool UNiagaraDataInterfaceChaosDestruction::GetFunctionHLSL(const FNiagaraDataInterfaceGPUParamInfo& ParamInfo, const FNiagaraDataInterfaceGeneratedFunction& FunctionInfo, int FunctionInstanceIndex, FString& OutHLSL)
 {
-	if (FunctionInfo.DefinitionName == GetPositionName)
-	{
-		static const TCHAR *Format = TEXT(R"(
-			void {FunctionName}(in int ParticleID, out float3 Out_Position) 
-			{
-				ParticleID -= LastSpawnedPointID_{Symbol} + 1;
-				Out_Position = PositionBuffer_{Symbol}[ParticleID];
-			}
-		)");
-
-		TMap<FString, FStringFormatArg> Args = {
-			{ TEXT("FunctionName"), FunctionInfo.InstanceName },
-			{ TEXT("Symbol"), ParamInfo.DataInterfaceHLSLSymbol },
-		};
-		OutHLSL += FString::Format(Format, Args);
-		return true;
-	}
-	else if (FunctionInfo.DefinitionName == GetNormalName)
-	{
-		static const TCHAR *Format = TEXT(R"(
-			void {FunctionName}(in int ParticleID, out float3 Out_Normal) 
-			{
-				ParticleID -= LastSpawnedPointID_{Symbol} + 1;
-				Out_Normal = IncomingNormalBuffer_{Symbol}[ParticleID];
-			}
-		)");
-
-		TMap<FString, FStringFormatArg> Args = {
-			{ TEXT("FunctionName"), FunctionInfo.InstanceName },
-			{ TEXT("Symbol"), ParamInfo.DataInterfaceHLSLSymbol },
-		};
-		OutHLSL += FString::Format(Format, Args);
-		return true;
-	}
-	else if (FunctionInfo.DefinitionName == GetVelocityName)
-	{
-		static const TCHAR *Format = TEXT(R"(
-			void {FunctionName}(in int ParticleID, out float3 Out_Velocity) 
-			{
-				ParticleID -= LastSpawnedPointID_{Symbol} + 1;
-				Out_Velocity = VelocityBuffer_{Symbol}[ParticleID];
-			}
-		)");
-
-		TMap<FString, FStringFormatArg> Args = {
-			{ TEXT("FunctionName"), FunctionInfo.InstanceName },
-			{ TEXT("Symbol"), ParamInfo.DataInterfaceHLSLSymbol },
-		};
-		OutHLSL += FString::Format(Format, Args);
-		return true;
-	}
-	else if (FunctionInfo.DefinitionName == GetAngularVelocityName)
-	{
-		static const TCHAR *Format = TEXT(R"(
-			void {FunctionName}(in int ParticleID, out float3 Out_AngularVelocity) 
-			{
-				ParticleID -= LastSpawnedPointID_{Symbol} + 1;
-				Out_AngularVelocity = IncomingAngularVelocity1Buffer_{Symbol}[ParticleID];
-			}
-		)");
-
-		TMap<FString, FStringFormatArg> Args = {
-			{ TEXT("FunctionName"), FunctionInfo.InstanceName },
-			{ TEXT("Symbol"), ParamInfo.DataInterfaceHLSLSymbol },
-		};
-		OutHLSL += FString::Format(Format, Args);
-		return true;
-	}
-	else if (FunctionInfo.DefinitionName == GetExtentMinName)
-	{
-		static const TCHAR *Format = TEXT(R"(
-			void {FunctionName}(in int ParticleID, out float Out_ExtentMin) 
-			{
-				ParticleID -= LastSpawnedPointID_{Symbol} + 1;
-				Out_ExtentMin = ExtentMinBuffer_{Symbol}[ParticleID];
-			}
-		)");
-
-		TMap<FString, FStringFormatArg> Args = {
-			{ TEXT("FunctionName"), FunctionInfo.InstanceName },
-			{ TEXT("Symbol"), ParamInfo.DataInterfaceHLSLSymbol },
-		};
-		OutHLSL += FString::Format(Format, Args);
-		return true;
-	}
-	else if (FunctionInfo.DefinitionName == GetExtentMaxName)
-	{
-		static const TCHAR *Format = TEXT(R"(
-			void {FunctionName}(in int ParticleID, out float Out_ExtentMax) 
-			{
-				ParticleID -= LastSpawnedPointID_{Symbol} + 1;
-				Out_ExtentMax = ExtentMaxBuffer_{Symbol}[ParticleID];
-			}
-		)");
-
-		TMap<FString, FStringFormatArg> Args = {
-			{ TEXT("FunctionName"), FunctionInfo.InstanceName },
-			{ TEXT("Symbol"), ParamInfo.DataInterfaceHLSLSymbol },
-		};
-		OutHLSL += FString::Format(Format, Args);
-		return true;
-	}
-	else if (FunctionInfo.DefinitionName == GetVolumeName)
-	{
-		static const TCHAR *Format = TEXT(R"(
-			void {FunctionName}(in int ParticleID, out float Out_Volume) 
-			{
-				ParticleID -= LastSpawnedPointID_{Symbol} + 1;
-				Out_Volume = VolumeBuffer_{Symbol}[ParticleID];
-			}
-		)");
-
-		TMap<FString, FStringFormatArg> Args = {
-			{ TEXT("FunctionName"), FunctionInfo.InstanceName },
-			{ TEXT("Symbol"), ParamInfo.DataInterfaceHLSLSymbol },
-		};
-		OutHLSL += FString::Format(Format, Args);
-		return true;
-	}
-	else if (FunctionInfo.DefinitionName == GetParticleIdsToSpawnAtTimeName)
-	{
-
-		static const TCHAR *Format = TEXT(R"(
-			void {FunctionName}(in float Time, out int Out_Min, 
-			                                   out int Out_Max, 
-			                                   out int Out_Count) 
-			{
-				// This function cannot be called on the GPU, as all spawn scripts are run on the CPU..
-				// TODO: Find a way to warn/error about this.
-				Out_Count = 0;
-				Out_Min = 0;
-				Out_Max = 0;
-			}
-		)");
-
-		TMap<FString, FStringFormatArg> Args = {
-			{ TEXT("FunctionName"), FunctionInfo.InstanceName },
-			{ TEXT("Symbol"), ParamInfo.DataInterfaceHLSLSymbol },
-		};
-		OutHLSL += FString::Format(Format, Args);
-		return true;
-	}
-	else if (FunctionInfo.DefinitionName == GetPointTypeName)
-	{
-		static const TCHAR *Format = TEXT(R"(
-			void {FunctionName}(in int ParticleID, out int Out_PointType) 
-			{
-				// NOTE(mv): Not yet part of the CPU functionality.
-				Out_PointType = 0;
-			}
-		)");
-
-		TMap<FString, FStringFormatArg> Args = {
-			{ TEXT("FunctionName"), FunctionInfo.InstanceName },
-		};
-		OutHLSL += FString::Format(Format, Args);
-		return true;
-	}
-	else if (FunctionInfo.DefinitionName == GetColorName)
-	{
-		static const TCHAR *Format = TEXT(R"(
-			void {FunctionName}(in int ParticleID, out float4 Out_Color) 
-			{
-				ParticleID -= LastSpawnedPointID_{Symbol} + 1;
-				Out_Color = ColorBuffer_{Symbol}[ParticleID];
-			}
-		)");
-
-		TMap<FString, FStringFormatArg> Args = {
-			{ TEXT("FunctionName"), FunctionInfo.InstanceName },
-			{ TEXT("Symbol"), ParamInfo.DataInterfaceHLSLSymbol },
-		};
-		OutHLSL += FString::Format(Format, Args);
-		return true;
-	}
-	else if (FunctionInfo.DefinitionName == GetSolverTimeName)
-	{
-		static const TCHAR *Format = TEXT(R"(
-			void {FunctionName}(out float Out_SolverTime) 
-			{
-				Out_SolverTime = SolverTime_{Symbol};
-			}
-		)");
-
-		TMap<FString, FStringFormatArg> Args = {
-			{ TEXT("FunctionName"), FunctionInfo.InstanceName },
-			{ TEXT("Symbol"), ParamInfo.DataInterfaceHLSLSymbol },
-		};
-		OutHLSL += FString::Format(Format, Args);
-		return true;
-	}
-	else if (FunctionInfo.DefinitionName == GetDensityName)
-	{
-		static const TCHAR *Format = TEXT(R"(
-			void {FunctionName}(in int ParticleID, out float Out_Density) 
-			{
-				ParticleID -= LastSpawnedPointID_{Symbol} + 1;
-				Out_Density = DensityBuffer_{Symbol}[ParticleID];
-			}
-		)");
-
-		TMap<FString, FStringFormatArg> Args = {
-			{ TEXT("FunctionName"), FunctionInfo.InstanceName },
-			{ TEXT("Symbol"), ParamInfo.DataInterfaceHLSLSymbol },
-		};
-		OutHLSL += FString::Format(Format, Args);
-		return true;
-	}
-	else if (FunctionInfo.DefinitionName == GetFrictionName)
-	{
-		static const TCHAR *Format = TEXT(R"(
-			void {FunctionName}(in int ParticleID, out float Out_Friction) 
-			{
-				ParticleID -= LastSpawnedPointID_{Symbol} + 1;
-				Out_Friction = FrictionBuffer_{Symbol}[ParticleID];
-			}
-		)");
-
-		TMap<FString, FStringFormatArg> Args = {
-			{ TEXT("FunctionName"), FunctionInfo.InstanceName },
-			{ TEXT("Symbol"), ParamInfo.DataInterfaceHLSLSymbol },
-		};
-		OutHLSL += FString::Format(Format, Args);
-		return true;
-	}
-	else if (FunctionInfo.DefinitionName == GetRestitutionName)
-	{
-		static const TCHAR *Format = TEXT(R"(
-			void {FunctionName}(in int ParticleID, out float Out_Restitution) 
-			{
-				ParticleID -= LastSpawnedPointID_{Symbol} + 1;
-				Out_Restitution = RestitutionBuffer_{Symbol}[ParticleID];
-			}
-		)");
-
-		TMap<FString, FStringFormatArg> Args = {
-			{ TEXT("FunctionName"), FunctionInfo.InstanceName },
-			{ TEXT("Symbol"), ParamInfo.DataInterfaceHLSLSymbol },
-		};
-		OutHLSL += FString::Format(Format, Args);
-		return true;
-	}
-	else if (FunctionInfo.DefinitionName == GetTransformName)
-	{
-		static const TCHAR *Format = TEXT(R"(
-			void {FunctionName}(in int ParticleID, out float3 Out_Translation, out float4 Out_Rotation, out float3 Out_Scale) 
-			{
-				ParticleID -= LastSpawnedPointID_{Symbol} + 1;
-				Out_Translation = TransformTranslationBuffer_{Symbol}[ParticleID];
-				Out_Rotation = TransformRotationBuffer_{Symbol}[ParticleID];
-				Out_Scale = TransformScaleBuffer_{Symbol}[ParticleID];
-			}
-		)");
-
-		TMap<FString, FStringFormatArg> Args = {
-			{ TEXT("FunctionName"), FunctionInfo.InstanceName },
-			{ TEXT("Symbol"), ParamInfo.DataInterfaceHLSLSymbol },
-		};
-		OutHLSL += FString::Format(Format, Args);
-		return true;
-	}
-	else if (FunctionInfo.DefinitionName == GetSizeName)
-	{
-		static const TCHAR *Format = TEXT(R"(
-			void {FunctionName}(in int ParticleID, out float3 Out_Size) 
-			{
-				ParticleID -= LastSpawnedPointID_{Symbol} + 1;
-				Out_Size = BoundsBuffer_{Symbol}[ParticleID];
-			}
-		)");
-
-		TMap<FString, FStringFormatArg> Args = {
-			{ TEXT("FunctionName"), FunctionInfo.InstanceName },
-			{ TEXT("Symbol"), ParamInfo.DataInterfaceHLSLSymbol },
-		};
-		OutHLSL += FString::Format(Format, Args);
-		return true;
-	}
-	else if (FunctionInfo.DefinitionName == GetSurfaceTypeName)
-	{
-		static const TCHAR *Format = TEXT(R"(
-			void {FunctionName}(in int ParticleID, out int Out_SurfaceType) 
-			{
-				ParticleID -= LastSpawnedPointID_{Symbol} + 1;
-				Out_SurfaceType = SurfaceTypeBuffer_{Symbol}[ParticleID];
-			}
-		)");
-
-		TMap<FString, FStringFormatArg> Args = {
-			{ TEXT("FunctionName"), FunctionInfo.InstanceName },
-			{ TEXT("Symbol"), ParamInfo.DataInterfaceHLSLSymbol },
-		};
-		OutHLSL += FString::Format(Format, Args);
-		return true;
-	}
-	else if (FunctionInfo.DefinitionName == GetCollisionDataName)
-	{
-		static const TCHAR *Format = TEXT(R"(
-			void {FunctionName}(in int ParticleID, out float3 Out_Location, 
-			                                       out float3 Out_AccumulatedImpulse,
-			                                       out float3 Out_Normal,
-			                                       out float3 Out_Velocity1,
-			                                       out float3 Out_Velocity2,
-			                                       out float3 Out_AngularVelocity1,
-			                                       out float3 Out_AngularVelocity2,
-			                                       out float  Out_Mass1,
-			                                       out float  Out_Mass2,
-			                                       out float  Out_Time) 
-			{
-				ParticleID -= LastSpawnedPointID_{Symbol} + 1;
-				Out_Location = IncomingLocationBuffer_{Symbol}[ParticleID];
-				Out_AccumulatedImpulse = IncomingAccumulatedImpulseBuffer_{Symbol}[ParticleID];
-				Out_Normal = IncomingNormalBuffer_{Symbol}[ParticleID];
-				Out_Velocity1 = IncomingVelocity1Buffer_{Symbol}[ParticleID];
-				Out_Velocity2 = IncomingVelocity2Buffer_{Symbol}[ParticleID];
-				Out_AngularVelocity1 = IncomingAngularVelocity1Buffer_{Symbol}[ParticleID];
-				Out_AngularVelocity2 = IncomingAngularVelocity2Buffer_{Symbol}[ParticleID];
-				Out_Mass1 = IncomingMass1Buffer_{Symbol}[ParticleID];
-				Out_Mass2 = IncomingMass2Buffer_{Symbol}[ParticleID];
-				Out_Time = IncomingTimeBuffer_{Symbol}[ParticleID];
-			}
-		)");
-
-		TMap<FString, FStringFormatArg> Args = {
-			{ TEXT("FunctionName"), FunctionInfo.InstanceName },
-			{ TEXT("Symbol"), ParamInfo.DataInterfaceHLSLSymbol },
-		};
-		OutHLSL += FString::Format(Format, Args);
-		return true;
-	}
-	else if (FunctionInfo.DefinitionName == GetBreakingDataName)
-	{
-		static const TCHAR *Format = TEXT(R"(
-			void {FunctionName}(in int ParticleID, out float3 Out_Location,
-			                                       out float3 Out_Velocity,
-			                                       out float3 Out_AngularVelocity,
-			                                       out float  Out_Mass,
-			                                       out float  Out_Time) 
-			{
-				ParticleID -= LastSpawnedPointID_{Symbol} + 1;
-				Out_Location = IncomingLocationBuffer_{Symbol}[ParticleID];
-				Out_Velocity = IncomingVelocity1Buffer_{Symbol}[ParticleID];
-				Out_AngularVelocity = IncomingAngularVelocity1Buffer_{Symbol}[ParticleID];
-				Out_Mass = IncomingMass1Buffer_{Symbol}[ParticleID];
-				Out_Time = IncomingTimeBuffer_{Symbol}[ParticleID];
-			}
-		)");
-
-		TMap<FString, FStringFormatArg> Args = {
-			{ TEXT("FunctionName"), FunctionInfo.InstanceName },
-			{ TEXT("Symbol"), ParamInfo.DataInterfaceHLSLSymbol },
-		};
-		OutHLSL += FString::Format(Format, Args);
-		return true;
-	}
-	else if (FunctionInfo.DefinitionName == GetTrailingDataName)
-	{
-		static const TCHAR *Format = TEXT(R"(
-			void {FunctionName}(in int ParticleID, out float3 Out_Location,
-			                                       out float3 Out_Velocity,
-			                                       out float3 Out_AngularVelocity,
-			                                       out float  Out_Mass,
-			                                       out float  Out_Time) 
-			{
-				ParticleID -= LastSpawnedPointID_{Symbol} + 1;
-				Out_Location = IncomingLocationBuffer_{Symbol}[ParticleID];
-				Out_Velocity = IncomingVelocity1Buffer_{Symbol}[ParticleID];
-				Out_AngularVelocity = IncomingAngularVelocity1Buffer_{Symbol}[ParticleID];
-				Out_Mass = IncomingMass1Buffer_{Symbol}[ParticleID];
-				Out_Time = IncomingTimeBuffer_{Symbol}[ParticleID];
-			}
-		)");
-
-		TMap<FString, FStringFormatArg> Args = {
-			{ TEXT("FunctionName"), FunctionInfo.InstanceName },
-			{ TEXT("Symbol"), ParamInfo.DataInterfaceHLSLSymbol },
-		};
-		OutHLSL += FString::Format(Format, Args);
-		return true;
-	}
+	using namespace NDIChaosDestructionLocal;
+
+	static const TSet<FName> ValidGpuFunction =
+	{
+		GetPositionName,
+		GetNormalName,
+		GetVelocityName,
+		GetAngularVelocityName,
+		GetExtentMinName,
+		GetExtentMaxName,
+		GetVolumeName,
+		GetParticleIdsToSpawnAtTimeName,
+		GetPointTypeName,
+		GetColorName,
+		GetSolverTimeName,
+		GetDensityName,
+		GetFrictionName,
+		GetRestitutionName,
+		GetTransformName,
+		GetSizeName,
+		GetSurfaceTypeName,
+		GetCollisionDataName,
+		GetBreakingDataName,
+		GetTrailingDataName,
+	};
 
-	return false;
+	return ValidGpuFunction.Contains(FunctionInfo.DefinitionName);
 }
 #endif
 
-template<typename T>
-void LoadGPUBufferFromArray(FDynamicReadBuffer& Buffer,
-	const TArray<T>* Array,
-	const EPixelFormat PixelFormat,
-	FString BufferName)
+void UNiagaraDataInterfaceChaosDestruction::BuildShaderParameters(FNiagaraShaderParametersBuilder& ShaderParametersBuilder) const
 {
-	checkf(PixelFormat == PF_A32B32G32R32F ||
-		PixelFormat == PF_R32_FLOAT ||
-		PixelFormat == PF_G32R32F ||
-		PixelFormat == PF_R32_UINT ||
-		PixelFormat == PF_R32_SINT,
-		TEXT("Unsupported PixelFormat: %d"), PixelFormat);
-
-	// NOTE: float3's have to be padded, so we pass them as PF_A32B32G32R32F and handle them differently
-	bool bIsVector = (PixelFormat == PF_A32B32G32R32F && sizeof(T) == 3 * sizeof(float));
-
-	uint32 SizePerElement = bIsVector ? 4 * sizeof(float) : sizeof(T);
-
-	// If not initialized, or if we need to expand the backing data
-	if (Buffer.NumBytes == 0 || Buffer.NumBytes < Array->Num() * SizePerElement)
-	{
-		Buffer.Release();
-		Buffer.Initialize(SizePerElement, Array->Num(), PixelFormat, BUF_Dynamic);
-	}
-
-	Buffer.Lock();
-	if (bIsVector)
-	{
-		FVector4f* Data = (FVector4f*)Buffer.MappedBuffer;
-		check(Data); // TODO: Handle gracefully.
-
-		for (int32 i = 0; i < Array->Num(); i++) {
-			Data[i] = FVector4f((*Array)[i]);
-		}
-	}
-	else
-	{
-		T* Data = (T*)Buffer.MappedBuffer;
-		check(Data); // TODO: Handle gracefully.
-
-		for (int32 i = 0; i < Array->Num(); i++) {
-			Data[i] = (*Array)[i];
-		}
-	}
-
-	Buffer.Unlock();
+	ShaderParametersBuilder.AddNestedStruct<NDIChaosDestructionLocal::FShaderParameters>();
 }
 
-template <typename T>
-static void SetBuffer(FRHICommandList& CmdList,
-	const FShaderResourceParameter& Param,
-	FRHIComputeShader* Shader,
-	FDynamicReadBuffer& Buffer,
-	const TArray<T>& Array,
-	const EPixelFormat PixelFormat,
-	const TCHAR* BufferName)
+void UNiagaraDataInterfaceChaosDestruction::SetShaderParameters(const FNiagaraDataInterfaceSetShaderParametersContext& Context) const
 {
-	checkf(PixelFormat == PF_A32B32G32R32F ||
-		   PixelFormat == PF_R32_FLOAT || 
-		   PixelFormat == PF_G32R32F || 
-		   PixelFormat == PF_R32_UINT || 
-		   PixelFormat == PF_R32_SINT, 
-		   TEXT("Unsupported PixelFormat: %d"), PixelFormat);
-
-	// Skip unbound parameters, since we won't be reading any of them
-	if (!Param.IsBound()) return;
-	
-	// NOTE: float3's have to be padded, so we pass them as PF_A32B32G32R32F and handle them differently
-	bool bIsVector = (PixelFormat == PF_A32B32G32R32F) && (sizeof(T) == 3*sizeof(float));
-
-	uint32 SizePerElement = bIsVector ? 4*sizeof(float) : sizeof(T); 
+	FNiagaraDataInterfaceProxyChaosDestruction& DIProxy = Context.GetProxy<FNiagaraDataInterfaceProxyChaosDestruction>();
+	FNiagaraDIChaosDestruction_GPUData* InstanceData = DIProxy.SystemsToGPUInstanceData.Find(Context.GetSystemInstanceID());
+	NDIChaosDestructionLocal::FShaderParameters* ShaderParameters = Context.GetParameterNestedStruct<NDIChaosDestructionLocal::FShaderParameters>();
 
-	// If not initialized, or if we need to expand the backing data
-	if (Buffer.NumBytes == 0 || Buffer.NumBytes < Array.Num() * SizePerElement)
+	if (InstanceData != nullptr && (InstanceData->PositionArray.Num() > 0))
 	{
-		Buffer.Release();
-		Buffer.Initialize(BufferName, SizePerElement, Array.Num(), PixelFormat, BUF_Dynamic);
-	}
+		NDIChaosDestructionLocal::QueueBufferUpload<FVector3f>(Context, ShaderParameters->PositionBuffer, InstanceData->RDGPositionSRV, InstanceData->PositionArray, FVector3f::ZeroVector, PF_A32B32G32R32F, TEXT("ChaosGPUBuffer_Position"));
+		NDIChaosDestructionLocal::QueueBufferUpload<FVector3f>(Context, ShaderParameters->VelocityBuffer, InstanceData->RDGVelocitySRV, InstanceData->VelocityArray, FVector3f::ZeroVector, PF_A32B32G32R32F, TEXT("ChaosGPUBuffer_Velocity"));
+		NDIChaosDestructionLocal::QueueBufferUpload<float>(Context, ShaderParameters->ExtentMinBuffer, InstanceData->RDGExtentMinBufferSRV, InstanceData->ExtentMinArray, 0.0f, PF_R32_FLOAT, TEXT("ChaosGPUBuffer_ExtentMin"));
+		NDIChaosDestructionLocal::QueueBufferUpload<float>(Context, ShaderParameters->ExtentMaxBuffer, InstanceData->RDGExtentMaxBufferSRV, InstanceData->ExtentMaxArray, 0.0f, PF_R32_FLOAT, TEXT("ChaosGPUBuffer_ExtentMax"));
+		NDIChaosDestructionLocal::QueueBufferUpload<float>(Context, ShaderParameters->VolumeBuffer, InstanceData->RDGVolumeBufferSRV, InstanceData->VolumeArray, 0.0f, PF_R32_FLOAT, TEXT("ChaosGPUBuffer_Volume"));
+		NDIChaosDestructionLocal::QueueBufferUpload<int>(Context, ShaderParameters->SolverIDBuffer, InstanceData->RDGSolverIDBufferSRV, InstanceData->SolverIDArray, 0, PF_R32_SINT, TEXT("ChaosGPUBuffer_SolverID"));
+		NDIChaosDestructionLocal::QueueBufferUpload<float>(Context, ShaderParameters->DensityBuffer, InstanceData->RDGDensityBufferSRV, InstanceData->DensityArray, 0.0f, PF_R32_FLOAT, TEXT("ChaosGPUBuffer_Density"));
+		NDIChaosDestructionLocal::QueueBufferUpload<float>(Context, ShaderParameters->FrictionBuffer, InstanceData->RDGFrictionBufferSRV, InstanceData->FrictionArray, 0.0f, PF_R32_FLOAT, TEXT("ChaosGPUBuffer_Friction"));
+		NDIChaosDestructionLocal::QueueBufferUpload<float>(Context, ShaderParameters->RestitutionBuffer, InstanceData->RDGRestitutionBufferSRV, InstanceData->RestitutionArray, 0.0f, PF_R32_FLOAT, TEXT("ChaosGPUBuffer_Restitution"));
+		NDIChaosDestructionLocal::QueueBufferUpload<int>(Context, ShaderParameters->SurfaceTypeBuffer, InstanceData->RDGSurfaceTypeBufferSRV, InstanceData->SurfaceTypeArray, 0, PF_R32_SINT, TEXT("ChaosGPUBuffer_SurfaceType"));
+		NDIChaosDestructionLocal::QueueBufferUpload<FLinearColor>(Context, ShaderParameters->ColorBuffer, InstanceData->RDGColorBufferSRV, InstanceData->ColorArray, FLinearColor::Black, PF_A32B32G32R32F, TEXT("ChaosGPUBuffer_Color"));
 
-	Buffer.Lock();
+		NDIChaosDestructionLocal::QueueBufferUpload<FVector3f>(Context, ShaderParameters->IncomingLocationBuffer, InstanceData->RDGIncomingLocationBufferSRV, InstanceData->IncomingLocationArray, FVector3f::ZeroVector, PF_A32B32G32R32F, TEXT("ChaosGPUBuffer_IncomingLocation"));
+		NDIChaosDestructionLocal::QueueBufferUpload<FVector3f>(Context, ShaderParameters->IncomingAccumulatedImpulseBuffer, InstanceData->RDGIncomingAccumulatedImpulseBufferSRV, InstanceData->IncomingAccumulatedImpulseArray, FVector3f::ZeroVector, PF_A32B32G32R32F, TEXT("ChaosGPUBuffer_IncomingAccumulatedImpulse"));
+		NDIChaosDestructionLocal::QueueBufferUpload<FVector3f>(Context, ShaderParameters->IncomingNormalBuffer, InstanceData->RDGIncomingNormalBufferSRV, InstanceData->IncomingNormalArray, FVector3f::ZeroVector, PF_A32B32G32R32F, TEXT("ChaosGPUBuffer_IncomingNormal"));
+		NDIChaosDestructionLocal::QueueBufferUpload<FVector3f>(Context, ShaderParameters->IncomingVelocity1Buffer, InstanceData->RDGIncomingVelocity1BufferSRV, InstanceData->IncomingVelocity1Array, FVector3f::ZeroVector, PF_A32B32G32R32F, TEXT("ChaosGPUBuffer_IncomingVelocity1"));
+		NDIChaosDestructionLocal::QueueBufferUpload<FVector3f>(Context, ShaderParameters->IncomingVelocity2Buffer, InstanceData->RDGIncomingVelocity2BufferSRV, InstanceData->IncomingVelocity2Array, FVector3f::ZeroVector, PF_A32B32G32R32F, TEXT("ChaosGPUBuffer_IncomingVelocity2"));
+		NDIChaosDestructionLocal::QueueBufferUpload<FVector3f>(Context, ShaderParameters->IncomingAngularVelocity1Buffer, InstanceData->RDGIncomingAngularVelocity1BufferSRV, InstanceData->IncomingAngularVelocity1Array, FVector3f::ZeroVector, PF_A32B32G32R32F, TEXT("ChaosGPUBuffer_IncomingAngualrVelocity1"));
+		NDIChaosDestructionLocal::QueueBufferUpload<FVector3f>(Context, ShaderParameters->IncomingAngularVelocity2Buffer, InstanceData->RDGIncomingAngularVelocity2BufferSRV, InstanceData->IncomingAngularVelocity2Array, FVector3f::ZeroVector, PF_A32B32G32R32F, TEXT("ChaosGPUBuffer_IncomingAngualrVelocity2"));
+		NDIChaosDestructionLocal::QueueBufferUpload<float>(Context, ShaderParameters->IncomingMass1Buffer, InstanceData->RDGIncomingMass1BufferSRV, InstanceData->IncomingMass1Array, 0.0f, PF_R32_FLOAT, TEXT("ChaosGPUBuffer_IncomingMass1"));
+		NDIChaosDestructionLocal::QueueBufferUpload<float>(Context, ShaderParameters->IncomingMass2Buffer, InstanceData->RDGIncomingMass2BufferSRV, InstanceData->IncomingMass2Array, 0.0f, PF_R32_FLOAT, TEXT("ChaosGPUBuffer_IncomingMass2"));
+		NDIChaosDestructionLocal::QueueBufferUpload<float>(Context, ShaderParameters->IncomingTimeBuffer, InstanceData->RDGIncomingTimeBufferSRV, InstanceData->IncomingTimeArray, 0.0f, PF_R32_FLOAT, TEXT("ChaosGPUBuffer_IncomingTime"));
 
-	static_assert(TIsPODType<T>::Value, "Not POD.");
-	
-	// NOTE: Reading from `Array` is not thread safe since it belongs to the simulation thread.
-	if (bIsVector)
-	{
-		FVector4f* Data = (FVector4f*)Buffer.MappedBuffer;
-		check(Data); // TODO: Handle gracefully.
+		NDIChaosDestructionLocal::QueueBufferUpload<FVector3f>(Context, ShaderParameters->TransformTranslationBuffer, InstanceData->RDGTransformTranslationBufferSRV, InstanceData->TransformTranslationArray, FVector3f::ZeroVector, PF_A32B32G32R32F, TEXT("ChaosGPUBuffer_TransformTranslation"));
+		NDIChaosDestructionLocal::QueueBufferUpload<FQuat4f>(Context, ShaderParameters->TransformRotationBuffer, InstanceData->RDGTransformRotationBufferSRV, InstanceData->TransformRotationArray, FQuat4f::Identity, PF_A32B32G32R32F, TEXT("ChaosGPUBuffer_TransformRotation"));
+		NDIChaosDestructionLocal::QueueBufferUpload<FVector3f>(Context, ShaderParameters->TransformScaleBuffer, InstanceData->RDGTransformScaleBufferSRV, InstanceData->TransformScaleArray, FVector3f::ZeroVector, PF_A32B32G32R32F, TEXT("ChaosGPUBuffer_TransformScale"));
+		NDIChaosDestructionLocal::QueueBufferUpload<FVector3f>(Context, ShaderParameters->BoundsBuffer, InstanceData->RDGBoundsBufferSRV, InstanceData->BoundsArray, FVector3f::ZeroVector, PF_A32B32G32R32F, TEXT("ChaosGPUBuffer_TransformBounds"));
 
-		for (int32 i = 0; i < Array.Num(); i++) {
-			FPlatformMemory::Memcpy(&Data[i], &Array[i], 3 * sizeof(float)); // TODO: Undefined data in fourth component
-		}
+		ShaderParameters->LastSpawnedPointID = InstanceData->LastSpawnedPointID;
+		ShaderParameters->SolverTime = InstanceData->SolverTime;
 	}
 	else
 	{
-		T* Data = (T*)Buffer.MappedBuffer;
-		check(Data); // TODO: Handle gracefully.
-		uint32 SizeToCopy = Array.Num() * SizePerElement;
-		check(Buffer.NumBytes >= SizeToCopy);
-		if (Array.Num() > 0)
-		{
-			FPlatformMemory::Memcpy(Data, &Array[0], SizeToCopy);
-		}
+		NDIChaosDestructionLocal::GetEmptyBuffer(Context, ShaderParameters->PositionBuffer, PF_A32B32G32R32F);
+		NDIChaosDestructionLocal::GetEmptyBuffer(Context, ShaderParameters->VelocityBuffer, PF_A32B32G32R32F);
+		NDIChaosDestructionLocal::GetEmptyBuffer(Context, ShaderParameters->ExtentMinBuffer, PF_R32_FLOAT);
+		NDIChaosDestructionLocal::GetEmptyBuffer(Context, ShaderParameters->ExtentMaxBuffer, PF_R32_FLOAT);
+		NDIChaosDestructionLocal::GetEmptyBuffer(Context, ShaderParameters->VolumeBuffer, PF_R32_FLOAT);
+		NDIChaosDestructionLocal::GetEmptyBuffer(Context, ShaderParameters->SolverIDBuffer, PF_R32_SINT);
+		NDIChaosDestructionLocal::GetEmptyBuffer(Context, ShaderParameters->DensityBuffer, PF_R32_FLOAT);
+		NDIChaosDestructionLocal::GetEmptyBuffer(Context, ShaderParameters->FrictionBuffer, PF_R32_FLOAT);
+		NDIChaosDestructionLocal::GetEmptyBuffer(Context, ShaderParameters->RestitutionBuffer, PF_R32_FLOAT);
+		NDIChaosDestructionLocal::GetEmptyBuffer(Context, ShaderParameters->SurfaceTypeBuffer, PF_R32_SINT);
+		NDIChaosDestructionLocal::GetEmptyBuffer(Context, ShaderParameters->ColorBuffer, PF_A32B32G32R32F);
+
+		NDIChaosDestructionLocal::GetEmptyBuffer(Context, ShaderParameters->IncomingLocationBuffer, PF_A32B32G32R32F);
+		NDIChaosDestructionLocal::GetEmptyBuffer(Context, ShaderParameters->IncomingAccumulatedImpulseBuffer, PF_A32B32G32R32F);
+		NDIChaosDestructionLocal::GetEmptyBuffer(Context, ShaderParameters->IncomingNormalBuffer, PF_A32B32G32R32F);
+		NDIChaosDestructionLocal::GetEmptyBuffer(Context, ShaderParameters->IncomingVelocity1Buffer, PF_A32B32G32R32F);
+		NDIChaosDestructionLocal::GetEmptyBuffer(Context, ShaderParameters->IncomingVelocity2Buffer, PF_A32B32G32R32F);
+		NDIChaosDestructionLocal::GetEmptyBuffer(Context, ShaderParameters->IncomingAngularVelocity1Buffer, PF_A32B32G32R32F);
+		NDIChaosDestructionLocal::GetEmptyBuffer(Context, ShaderParameters->IncomingAngularVelocity2Buffer, PF_A32B32G32R32F);
+		NDIChaosDestructionLocal::GetEmptyBuffer(Context, ShaderParameters->IncomingMass1Buffer, PF_R32_FLOAT);
+		NDIChaosDestructionLocal::GetEmptyBuffer(Context, ShaderParameters->IncomingMass2Buffer, PF_R32_FLOAT);
+		NDIChaosDestructionLocal::GetEmptyBuffer(Context, ShaderParameters->IncomingTimeBuffer, PF_R32_FLOAT);
+
+		NDIChaosDestructionLocal::GetEmptyBuffer(Context, ShaderParameters->TransformTranslationBuffer, PF_A32B32G32R32F);
+		NDIChaosDestructionLocal::GetEmptyBuffer(Context, ShaderParameters->TransformRotationBuffer, PF_A32B32G32R32F);
+		NDIChaosDestructionLocal::GetEmptyBuffer(Context, ShaderParameters->TransformScaleBuffer, PF_A32B32G32R32F);
+		NDIChaosDestructionLocal::GetEmptyBuffer(Context, ShaderParameters->BoundsBuffer, PF_A32B32G32R32F);
+
+		ShaderParameters->LastSpawnedPointID = 0;
+		ShaderParameters->SolverTime = 0.0f;
 	}
-
-	Buffer.Unlock();
-
-	SetSRVParameter(CmdList, Shader, Param, Buffer.SRV);
 }
 
 void UNiagaraDataInterfaceChaosDestruction::PushToRenderThreadImpl()
@@ -3790,140 +3504,31 @@ void UNiagaraDataInterfaceChaosDestruction::ProvidePerInstanceDataForRenderThrea
 	DataToPass->SolverTime = GetSolverTime();
 	DataToPass->LastSpawnedPointID = GetLastSpawnedPointID();
 
-	if (InstanceData->PositionArray.Num() > 0)
-	{
-		DataToPass->PositionArray = new TArray<FVector3f>();
-		DataToPass->PositionArray->AddUninitialized(InstanceData->PositionArray.Num());
-		for (int i = 0; i < InstanceData->PositionArray.Num(); ++i)
-		{
-			(*DataToPass->PositionArray)[i] = (FVector3f)InstanceData->PositionArray[i];
-		}
-	}
-
-	if (InstanceData->VelocityArray.Num() > 0)
-	{
-		DataToPass->VelocityArray = new TArray<FVector3f>();
-		DataToPass->VelocityArray->AddUninitialized(InstanceData->VelocityArray.Num());
-		for (int i = 0; i < InstanceData->VelocityArray.Num(); ++i)
-		{
-			(*DataToPass->VelocityArray)[i] = (FVector3f)InstanceData->VelocityArray[i];
-		}
-	}
-
-	if (InstanceData->ExtentMinArray.Num() > 0)
-	{
-		DataToPass->ExtentMinArray = new TArray<float>(InstanceData->ExtentMinArray);
-	}
-
-	if (InstanceData->ExtentMaxArray.Num() > 0)
-	{
-		DataToPass->ExtentMaxArray = new TArray<float>(InstanceData->ExtentMaxArray);
-	}
-
-	if (InstanceData->VolumeArray.Num() > 0)
-	{
-		DataToPass->VolumeArray = new TArray<float>(InstanceData->VolumeArray);
-	}
-
-	if (InstanceData->SolverIDArray.Num() > 0)
-	{
-		DataToPass->SolverIDArray = new TArray<int32>(InstanceData->SolverIDArray);
-	}
-
-	if (InstanceData->DensityArray.Num() > 0)
-	{
-		DataToPass->DensityArray = new TArray<float>(InstanceData->DensityArray);
-	}
-
-	if (InstanceData->FrictionArray.Num() > 0)
-	{
-		DataToPass->FrictionArray = new TArray<float>(InstanceData->FrictionArray);
-	}
-
-	if (InstanceData->RestitutionArray.Num() > 0)
-	{
-		DataToPass->RestitutionArray = new TArray<float>(InstanceData->RestitutionArray);
-	}
-
-	if (InstanceData->TransformTranslationArray.Num() > 0)
-	{
-		DataToPass->TransformTranslationArray = new TArray<FVector>(InstanceData->TransformTranslationArray);
-	}
-
-	if (InstanceData->TransformRotationArray.Num() > 0)
-	{
-		DataToPass->TransformRotationArray = new TArray<FQuat>(InstanceData->TransformRotationArray);
-	}
-
-	if (InstanceData->TransformScaleArray.Num() > 0)
-	{
-		DataToPass->TransformScaleArray = new TArray<FVector>(InstanceData->TransformScaleArray);
-	}
-
-	if (InstanceData->BoundsArray.Num() > 0)
-	{
-		DataToPass->BoundsArray = new TArray<FVector>(InstanceData->BoundsArray);
-	}
-
-	if (InstanceData->SurfaceTypeArray.Num() > 0)
-	{
-		DataToPass->SurfaceTypeArray = new TArray<int32>(InstanceData->SurfaceTypeArray);
-	}
-
-	if (InstanceData->ColorArray.Num() > 0)
-	{
-		DataToPass->ColorArray = new TArray<FLinearColor>(InstanceData->ColorArray);
-	}
-
-	if (InstanceData->IncomingLocationArray.Num() > 0)
-	{
-		DataToPass->IncomingLocationArray = new TArray<FVector>(InstanceData->IncomingLocationArray);
-	}
-
-	if (InstanceData->IncomingAccumulatedImpulseArray.Num() > 0)
-	{
-		DataToPass->IncomingAccumulatedImpulseArray = new TArray<FVector>(InstanceData->IncomingAccumulatedImpulseArray);
-	}
-
-	if (InstanceData->IncomingNormalArray.Num() > 0)
-	{
-		DataToPass->IncomingNormalArray = new TArray<FVector>(InstanceData->IncomingNormalArray);
-	}
-
-	if (InstanceData->IncomingVelocity1Array.Num() > 0)
-	{
-		DataToPass->IncomingVelocity1Array = new TArray<FVector>(InstanceData->IncomingVelocity1Array);
-	}
-
-	if (InstanceData->IncomingVelocity2Array.Num() > 0)
-	{
-		DataToPass->IncomingVelocity2Array = new TArray<FVector>(InstanceData->IncomingVelocity2Array);
-	}
-
-	if (InstanceData->IncomingAngularVelocity1Array.Num() > 0)
-	{
-		DataToPass->IncomingAngularVelocity1Array = new TArray<FVector>(InstanceData->IncomingAngularVelocity1Array);
-	}
-
-	if (InstanceData->IncomingAngularVelocity2Array.Num() > 0)
-	{
-		DataToPass->IncomingAngularVelocity2Array = new TArray<FVector>(InstanceData->IncomingAngularVelocity2Array);
-	}
-
-	if (InstanceData->IncomingMass1Array.Num() > 0)
-	{
-		DataToPass->IncomingMass1Array = new TArray<float>(InstanceData->IncomingMass1Array);
-	}
-
-	if (InstanceData->IncomingMass2Array.Num() > 0)
-	{
-		DataToPass->IncomingMass2Array = new TArray<float>(InstanceData->IncomingMass2Array);
-	}
-
-	if (InstanceData->IncomingTimeArray.Num() > 0)
-	{
-		DataToPass->IncomingTimeArray = new TArray<float>(InstanceData->IncomingTimeArray);
-	}
+	DataToPass->PositionArray = NDIChaosDestructionLocal::CopyVectorArray(InstanceData->PositionArray);
+	DataToPass->VelocityArray = NDIChaosDestructionLocal::CopyVectorArray(InstanceData->VelocityArray);
+	DataToPass->ExtentMinArray = InstanceData->ExtentMinArray;
+	DataToPass->ExtentMaxArray = InstanceData->ExtentMaxArray;
+	DataToPass->VolumeArray = InstanceData->VolumeArray;
+	DataToPass->SolverIDArray = InstanceData->SolverIDArray;
+	DataToPass->DensityArray = InstanceData->DensityArray;
+	DataToPass->FrictionArray = InstanceData->FrictionArray;
+	DataToPass->RestitutionArray = InstanceData->RestitutionArray;
+	DataToPass->TransformTranslationArray = NDIChaosDestructionLocal::CopyVectorArray(InstanceData->TransformTranslationArray);
+	DataToPass->TransformRotationArray = NDIChaosDestructionLocal::CopyQuatArray(InstanceData->TransformRotationArray);
+	DataToPass->TransformScaleArray = NDIChaosDestructionLocal::CopyVectorArray(InstanceData->TransformScaleArray);
+	DataToPass->BoundsArray = NDIChaosDestructionLocal::CopyVectorArray(InstanceData->BoundsArray);
+	DataToPass->SurfaceTypeArray = InstanceData->SurfaceTypeArray;
+	DataToPass->ColorArray = InstanceData->ColorArray;
+	DataToPass->IncomingLocationArray = NDIChaosDestructionLocal::CopyVectorArray(InstanceData->IncomingLocationArray);
+	DataToPass->IncomingAccumulatedImpulseArray = NDIChaosDestructionLocal::CopyVectorArray(InstanceData->IncomingAccumulatedImpulseArray);
+	DataToPass->IncomingNormalArray = NDIChaosDestructionLocal::CopyVectorArray(InstanceData->IncomingNormalArray);
+	DataToPass->IncomingVelocity1Array = NDIChaosDestructionLocal::CopyVectorArray(InstanceData->IncomingVelocity1Array);
+	DataToPass->IncomingVelocity2Array = NDIChaosDestructionLocal::CopyVectorArray(InstanceData->IncomingVelocity2Array);
+	DataToPass->IncomingAngularVelocity1Array = NDIChaosDestructionLocal::CopyVectorArray(InstanceData->IncomingAngularVelocity1Array);
+	DataToPass->IncomingAngularVelocity2Array = NDIChaosDestructionLocal::CopyVectorArray(InstanceData->IncomingAngularVelocity2Array);
+	DataToPass->IncomingMass1Array = InstanceData->IncomingMass1Array;
+	DataToPass->IncomingMass2Array = InstanceData->IncomingMass2Array;
+	DataToPass->IncomingTimeArray = InstanceData->IncomingTimeArray;
 }
 
 void FNiagaraDataInterfaceProxyChaosDestruction::CreatePerInstanceData(const FNiagaraSystemInstanceID& SystemInstance)
@@ -3953,315 +3558,71 @@ void FNiagaraDataInterfaceProxyChaosDestruction::ConsumePerInstanceDataFromGameT
 		Data.SolverTime = InstanceData->SolverTime;
 		Data.LastSpawnedPointID = InstanceData->LastSpawnedPointID;
 
-		if (InstanceData->PositionArray)
-		{
-			Data.PositionArray = TArray<FVector3f>(MoveTemp(*InstanceData->PositionArray));		
-			//LoadGPUBufferFromArray(Data.GPUPositionBuffer, InstanceData->PositionArray, EPixelFormat::PF_A32B32G32R32F, FString(TEXT("PositionBuffer")));
-			delete InstanceData->PositionArray;
-		}
-
-		if (InstanceData->VelocityArray)
-		{
-			Data.VelocityArray = TArray<FVector3f>(MoveTemp(*InstanceData->VelocityArray));
-			//LoadGPUBufferFromArray(Data.GPUVelocityBuffer, InstanceData->VelocityArray, EPixelFormat::PF_A32B32G32R32F, FString(TEXT("VelocityBuffer")));
-			delete InstanceData->VelocityArray;
-		}
-	
-		if (InstanceData->ExtentMinArray)
-		{
-			Data.ExtentMinArray = TArray<float>(MoveTemp(*InstanceData->ExtentMinArray));
-			//LoadGPUBufferFromArray(Data.GPUExtentMinBuffer, InstanceData->ExtentMinArray, EPixelFormat::PF_R32_FLOAT, FString(TEXT("ExtentMinBuffer")));
-			delete InstanceData->ExtentMinArray;
-		}
-	
-		if (InstanceData->ExtentMaxArray)
-		{
-			Data.ExtentMaxArray = TArray<float>(MoveTemp(*InstanceData->ExtentMaxArray));
-			//LoadGPUBufferFromArray(Data.GPUExtentMaxBuffer, InstanceData->ExtentMaxArray, EPixelFormat::PF_R32_FLOAT, FString(TEXT("ExtentMaxBuffer")));
-			delete InstanceData->ExtentMaxArray;
-		}
-
-		if (InstanceData->VolumeArray)
-		{
-			Data.VolumeArray = TArray<float>(MoveTemp(*InstanceData->VolumeArray));
-			//LoadGPUBufferFromArray(Data.GPUVolumeBuffer, InstanceData->VolumeArray, EPixelFormat::PF_R32_FLOAT, FString(TEXT("VolumeBuffer")));
-			delete InstanceData->VolumeArray;
-		}
-
-		if (InstanceData->SolverIDArray)
-		{
-			Data.SolverIDArray = TArray<int32>(MoveTemp(*InstanceData->SolverIDArray));
-			//LoadGPUBufferFromArray(Data.GPUSolverIDBuffer, InstanceData->SolverIDArray, EPixelFormat::PF_R32_SINT, FString(TEXT("SolverIDBuffer")));
-			delete InstanceData->SolverIDArray;
-		}
-
-		if (InstanceData->DensityArray)
-		{
-			Data.DensityArray = TArray<float>(MoveTemp(*InstanceData->DensityArray));
-			//LoadGPUBufferFromArray(Data.GPUDensityBuffer, InstanceData->DensityArray, EPixelFormat::PF_R32_FLOAT, FString(TEXT("DensityBuffer")));
-			delete InstanceData->DensityArray;
-		}
-
-		if (InstanceData->FrictionArray)
-		{
-			Data.FrictionArray = TArray<float>(MoveTemp(*InstanceData->FrictionArray));
-			//LoadGPUBufferFromArray(Data.GPUFrictionBuffer, InstanceData->FrictionArray, EPixelFormat::PF_R32_FLOAT, FString(TEXT("FrictionBuffer")));
-			delete InstanceData->FrictionArray;
-		}
+		Data.PositionArray = MoveTemp(InstanceData->PositionArray);
+		Data.VelocityArray = MoveTemp(InstanceData->VelocityArray);
+		Data.ExtentMinArray = MoveTemp(InstanceData->ExtentMinArray);
+		Data.ExtentMaxArray = MoveTemp(InstanceData->ExtentMaxArray);
+		Data.VolumeArray = MoveTemp(InstanceData->VolumeArray);
+		Data.SolverIDArray = MoveTemp(InstanceData->SolverIDArray);
+		Data.DensityArray = MoveTemp(InstanceData->DensityArray);
+		Data.FrictionArray = MoveTemp(InstanceData->FrictionArray);
+		Data.RestitutionArray = MoveTemp(InstanceData->RestitutionArray);
+		Data.TransformTranslationArray = MoveTemp(InstanceData->TransformTranslationArray);
+		Data.TransformRotationArray = MoveTemp(InstanceData->TransformRotationArray);
+		Data.TransformScaleArray = MoveTemp(InstanceData->TransformScaleArray);
+		Data.BoundsArray = MoveTemp(InstanceData->BoundsArray);
+		Data.SurfaceTypeArray = MoveTemp(InstanceData->SurfaceTypeArray);
+		Data.ColorArray = MoveTemp(InstanceData->ColorArray);
+		Data.IncomingLocationArray = MoveTemp(InstanceData->IncomingLocationArray);
+		Data.IncomingAccumulatedImpulseArray = MoveTemp(InstanceData->IncomingAccumulatedImpulseArray);
+		Data.IncomingNormalArray = MoveTemp(InstanceData->IncomingNormalArray);
+		Data.IncomingVelocity1Array = MoveTemp(InstanceData->IncomingVelocity1Array);
+		Data.IncomingVelocity2Array = MoveTemp(InstanceData->IncomingVelocity2Array);
+		Data.IncomingAngularVelocity1Array = MoveTemp(InstanceData->IncomingAngularVelocity1Array);
+		Data.IncomingAngularVelocity2Array = MoveTemp(InstanceData->IncomingAngularVelocity2Array);
+		Data.IncomingMass1Array = MoveTemp(InstanceData->IncomingMass1Array);
+		Data.IncomingMass2Array = MoveTemp(InstanceData->IncomingMass2Array);
+		Data.IncomingTimeArray = MoveTemp(InstanceData->IncomingTimeArray);
 
-		if (InstanceData->RestitutionArray)
-		{
-			Data.RestitutionArray = TArray<float>(MoveTemp(*InstanceData->RestitutionArray));
-			//LoadGPUBufferFromArray(Data.GPURestitutionBuffer, InstanceData->RestitutionArray, EPixelFormat::PF_R32_FLOAT, FString(TEXT("RestitutionBuffer")));
-			delete InstanceData->RestitutionArray;
-		}
-
-		if (InstanceData->TransformTranslationArray)
-		{
-			Data.TransformTranslationArray = TArray<FVector>(MoveTemp(*InstanceData->TransformTranslationArray));
-			//LoadGPUBufferFromArray(Data.GPUSurfaceTypeBuffer, InstanceData->TransformTranslationArray, EPixelFormat::PF_A32B32G32R32F, FString(TEXT("TransformTranslationBuffer")));
-			delete InstanceData->TransformTranslationArray;
-		}
-
-		if (InstanceData->TransformRotationArray)
-		{
-			Data.TransformRotationArray = TArray<FQuat>(MoveTemp(*InstanceData->TransformRotationArray));
-			//LoadGPUBufferFromArray(Data.GPUSurfaceTypeBuffer, InstanceData->TransformRotationArray, EPixelFormat::PF_A32B32G32R32F, FString(TEXT("TransformRotationBuffer")));
-			delete InstanceData->TransformRotationArray;
-		}
-
-		if (InstanceData->TransformScaleArray)
-		{
-			Data.TransformScaleArray = TArray<FVector>(MoveTemp(*InstanceData->TransformScaleArray));
-			//LoadGPUBufferFromArray(Data.GPUSurfaceTypeBuffer, InstanceData->TransformScaleArray, EPixelFormat::PF_A32B32G32R32F, FString(TEXT("TransformScaleBuffer")));
-			delete InstanceData->TransformScaleArray;
-		}
-
-		if (InstanceData->BoundsArray)
-		{
-			Data.BoundsArray = TArray<FVector>(MoveTemp(*InstanceData->BoundsArray));
-			//LoadGPUBufferFromArray(Data.GPUBoundsBuffer, InstanceData->BoundsArray, EPixelFormat::PF_A32B32G32R32F, FString(TEXT("BoundsBuffer")));
-			delete InstanceData->BoundsArray;
-		}
-
-		if (InstanceData->SurfaceTypeArray)
-		{
-			Data.SurfaceTypeArray = TArray<int32>(MoveTemp(*InstanceData->SurfaceTypeArray));
-			//LoadGPUBufferFromArray(Data.GPUSurfaceTypeBuffer, InstanceData->SurfaceTypeArray, EPixelFormat::PF_R32_FLOAT, FString(TEXT("RestitutionBuffer")));
-			delete InstanceData->SurfaceTypeArray;
-		}
-
-		if (InstanceData->ColorArray)
-		{
-			Data.ColorArray = TArray<FLinearColor>(MoveTemp(*InstanceData->ColorArray));
-			//LoadGPUBufferFromArray(Data.GPUColorBuffer, InstanceData->ColorArray, EPixelFormat::PF_A32B32G32R32F, FString(TEXT("ColorBuffer")));
-			delete InstanceData->ColorArray;
-		}
-
-		if (InstanceData->IncomingLocationArray)
-		{
-			Data.IncomingLocationArray = TArray<FVector>(MoveTemp(*InstanceData->IncomingLocationArray));
-			//LoadGPUBufferFromArray(Data.GPUIncomingLocationBuffer, InstanceData->IncomingLocationArray, EPixelFormat::PF_A32B32G32R32F, FString(TEXT("IncomingLocationBuffer")));
-			delete InstanceData->IncomingLocationArray;
-		}
-
-		if (InstanceData->IncomingAccumulatedImpulseArray)
-		{
-			Data.IncomingAccumulatedImpulseArray = TArray<FVector>(MoveTemp(*InstanceData->IncomingAccumulatedImpulseArray));
-			//LoadGPUBufferFromArray(Data.GPUIncomingAccumulatedImpulseBuffer, InstanceData->IncomingAccumulatedImpulseArray, EPixelFormat::PF_A32B32G32R32F, FString(TEXT("IncomingAccumulatedImpulseBuffer")));
-			delete InstanceData->IncomingAccumulatedImpulseArray;
-		}
-
-		if (InstanceData->IncomingNormalArray)
-		{
-			Data.IncomingNormalArray = TArray<FVector>(MoveTemp(*InstanceData->IncomingNormalArray));
-			//LoadGPUBufferFromArray(Data.GPUIncomingNormalBuffer, InstanceData->IncomingNormalArray, EPixelFormat::PF_A32B32G32R32F, FString(TEXT("IncomingNormalBuffer")));
-			delete InstanceData->IncomingNormalArray;
-		}
-
-		if (InstanceData->IncomingVelocity1Array)
-		{
-			Data.IncomingVelocity1Array = TArray<FVector>(MoveTemp(*InstanceData->IncomingVelocity1Array));
-			//LoadGPUBufferFromArray(Data.GPUIncomingVelocity1Buffer, InstanceData->IncomingVelocity1Array, EPixelFormat::PF_A32B32G32R32F, FString(TEXT("IncomingVelocity1Buffer")));
-			delete InstanceData->IncomingVelocity1Array;
-		}
-
-		if (InstanceData->IncomingVelocity2Array)
-		{
-			Data.IncomingVelocity2Array = TArray<FVector>(MoveTemp(*InstanceData->IncomingVelocity2Array));
-			//LoadGPUBufferFromArray(Data.GPUIncomingVelocity2Buffer, InstanceData->IncomingVelocity2Array, EPixelFormat::PF_A32B32G32R32F, FString(TEXT("IncomingVelocity2Buffer")));
-			delete InstanceData->IncomingVelocity2Array;
-		}
-
-		if (InstanceData->IncomingAngularVelocity1Array)
-		{
-			Data.IncomingAngularVelocity1Array = TArray<FVector>(MoveTemp(*InstanceData->IncomingAngularVelocity1Array));
-			//LoadGPUBufferFromArray(Data.GPUIncomingAngularVelocity1Buffer, InstanceData->IncomingAngularVelocity1Array, EPixelFormat::PF_A32B32G32R32F, FString(TEXT("IncomingAngularVelocity1Buffer")));
-			delete InstanceData->IncomingAngularVelocity1Array;
-		}
-
-		if (InstanceData->IncomingAngularVelocity2Array)
-		{
-			Data.IncomingAngularVelocity2Array = TArray<FVector>(MoveTemp(*InstanceData->IncomingAngularVelocity2Array));
-			//LoadGPUBufferFromArray(Data.GPUIncomingAngularVelocity2Buffer, InstanceData->IncomingAngularVelocity2Array, EPixelFormat::PF_A32B32G32R32F, FString(TEXT("IncomingAngularVelocity2Buffer")));
-			delete InstanceData->IncomingAngularVelocity2Array;
-		}
-	
-		if (InstanceData->IncomingMass1Array)
-		{
-			Data.IncomingMass1Array = TArray<float>(MoveTemp(*InstanceData->IncomingMass1Array));
-			//LoadGPUBufferFromArray(Data.GPUIncomingMass1Buffer, InstanceData->IncomingMass1Array, EPixelFormat::PF_R32_FLOAT, FString(TEXT("IncomingMass1Buffer")));
-			delete InstanceData->IncomingMass1Array;
-		}
-
-		if (InstanceData->IncomingMass2Array)
-		{
-			Data.IncomingMass2Array = TArray<float>(MoveTemp(*InstanceData->IncomingMass2Array));
-			//LoadGPUBufferFromArray(Data.GPUIncomingMass2Buffer, InstanceData->IncomingMass2Array, EPixelFormat::PF_R32_FLOAT, FString(TEXT("IncomingMass2Buffer")));
-			delete InstanceData->IncomingMass2Array;
-		}
-
-		if (InstanceData->IncomingTimeArray)
-		{
-			Data.IncomingTimeArray = TArray<float>(MoveTemp(*InstanceData->IncomingTimeArray));
-			//LoadGPUBufferFromArray(Data.GPUIncomingTimeBuffer, InstanceData->IncomingTimeArray, EPixelFormat::PF_R32_FLOAT, FString(TEXT("IncomingTimeBuffer")));
-			delete InstanceData->IncomingTimeArray;
-		}
 		InstanceData->~FNiagaraDIChaosDestruction_InstanceDataToPassToRT();
 	}
 }
 
-struct FNiagaraDataInterfaceParametersCS_ChaosDestruction :public FNiagaraDataInterfaceParametersCS
+void FNiagaraDataInterfaceProxyChaosDestruction::PostSimulate(const FNDIGpuComputePostSimulateContext& Context)
 {
-	DECLARE_TYPE_LAYOUT(FNiagaraDataInterfaceParametersCS_ChaosDestruction, NonVirtual);
-public:
-	void Bind(const FNiagaraDataInterfaceGPUParamInfo& ParameterInfo, const class FShaderParameterMap& ParameterMap)
-	{
-		PositionBuffer.Bind(ParameterMap, *("PositionBuffer_" + ParameterInfo.DataInterfaceHLSLSymbol));
-		VelocityBuffer.Bind(ParameterMap, *("VelocityBuffer_" + ParameterInfo.DataInterfaceHLSLSymbol));
-		ExtentMinBuffer.Bind(ParameterMap, *("ExtentMinBuffer_" + ParameterInfo.DataInterfaceHLSLSymbol));
-		ExtentMaxBuffer.Bind(ParameterMap, *("ExtentMaxBuffer_" + ParameterInfo.DataInterfaceHLSLSymbol));
-		VolumeBuffer.Bind(ParameterMap, *("VolumeBuffer_" + ParameterInfo.DataInterfaceHLSLSymbol));
-		SolverIDBuffer.Bind(ParameterMap, *("SolverIDBuffer_" + ParameterInfo.DataInterfaceHLSLSymbol));
-		DensityBuffer.Bind(ParameterMap, *("DensityBuffer_" + ParameterInfo.DataInterfaceHLSLSymbol));
-		FrictionBuffer.Bind(ParameterMap, *("FrictionBuffer_" + ParameterInfo.DataInterfaceHLSLSymbol));
-		RestitutionBuffer.Bind(ParameterMap, *("RestitutionBuffer_" + ParameterInfo.DataInterfaceHLSLSymbol));
-		SurfaceTypeBuffer.Bind(ParameterMap, *("SurfaceTypeBuffer_" + ParameterInfo.DataInterfaceHLSLSymbol));
-		ColorBuffer.Bind(ParameterMap, *("ColorBuffer_" + ParameterInfo.DataInterfaceHLSLSymbol));
-
-		IncomingLocationBuffer.Bind(ParameterMap, *("IncomingLocationBuffer_" + ParameterInfo.DataInterfaceHLSLSymbol));
-		IncomingAccumulatedImpulseBuffer.Bind(ParameterMap, *("IncomingAccumulatedImpulseBuffer_" + ParameterInfo.DataInterfaceHLSLSymbol));
-		IncomingNormalBuffer.Bind(ParameterMap, *("IncomingNormalBuffer_" + ParameterInfo.DataInterfaceHLSLSymbol));
-		IncomingVelocity1Buffer.Bind(ParameterMap, *("IncomingVelocity1Buffer_" + ParameterInfo.DataInterfaceHLSLSymbol));
-		IncomingVelocity2Buffer.Bind(ParameterMap, *("IncomingVelocity2Buffer_" + ParameterInfo.DataInterfaceHLSLSymbol));
-		IncomingAngularVelocity1Buffer.Bind(ParameterMap, *("IncomingAngularVelocity1Buffer_" + ParameterInfo.DataInterfaceHLSLSymbol));
-		IncomingAngularVelocity2Buffer.Bind(ParameterMap, *("IncomingAngularVelocity2Buffer_" + ParameterInfo.DataInterfaceHLSLSymbol));
-		IncomingMass1Buffer.Bind(ParameterMap, *("IncomingMass1Buffer_" + ParameterInfo.DataInterfaceHLSLSymbol));
-		IncomingMass2Buffer.Bind(ParameterMap, *("IncomingMass2Buffer_" + ParameterInfo.DataInterfaceHLSLSymbol));
-		IncomingTimeBuffer.Bind(ParameterMap, *("IncomingTimeBuffer_" + ParameterInfo.DataInterfaceHLSLSymbol));
-
-		TransformTranslationBuffer.Bind(ParameterMap, *("TransformTranslationBuffer_" + ParameterInfo.DataInterfaceHLSLSymbol));
-		TransformRotationBuffer.Bind(ParameterMap, *("TransformRotationBuffer_" + ParameterInfo.DataInterfaceHLSLSymbol));
-		TransformScaleBuffer.Bind(ParameterMap, *("TransformScaleBuffer_" + ParameterInfo.DataInterfaceHLSLSymbol));
-		BoundsBuffer.Bind(ParameterMap, *("BoundsBuffer_" + ParameterInfo.DataInterfaceHLSLSymbol));
-
-		LastSpawnedPointID.Bind(ParameterMap, *("LastSpawnedPointID_" + ParameterInfo.DataInterfaceHLSLSymbol));
-		SolverTime.Bind(ParameterMap, *("SolverTime_" + ParameterInfo.DataInterfaceHLSLSymbol));
-	}
-
-	void Set(FRHICommandList& RHICmdList, const FNiagaraDataInterfaceSetArgs& Context) const
+	// Clear out the transient data we may have generated if this is the final post simulate call
+	if (Context.IsFinalPostSimulate())
 	{
-		check(IsInRenderingThread());
-
-		FRHIComputeShader* ComputeShaderRHI = Context.Shader.GetComputeShader();
-		FNiagaraDataInterfaceProxyChaosDestruction* ChaosDestructionInterfaceProxy = static_cast<FNiagaraDataInterfaceProxyChaosDestruction*>(Context.DataInterface);
-		if (ChaosDestructionInterfaceProxy)
+		for (auto it = SystemsToGPUInstanceData.CreateIterator(); it; ++it)
 		{
-			FNiagaraDIChaosDestruction_GPUData* InstanceData = ChaosDestructionInterfaceProxy->SystemsToGPUInstanceData.Find(Context.SystemInstanceID);
-
-			ensure(InstanceData);
-
-			if (!InstanceData)
-			{
-				return;
-			}
-			if(InstanceData->PositionArray.Num() > 0)
-			{
-				SetBuffer(RHICmdList, PositionBuffer,    ComputeShaderRHI, InstanceData->GPUPositionBuffer,    InstanceData->PositionArray,    EPixelFormat::PF_A32B32G32R32F, TEXT("PositionBuffer"));
-				SetBuffer(RHICmdList, VelocityBuffer,    ComputeShaderRHI, InstanceData->GPUVelocityBuffer,    InstanceData->VelocityArray,    EPixelFormat::PF_A32B32G32R32F, TEXT("VelocityBuffer"));
-				SetBuffer(RHICmdList, ExtentMinBuffer,   ComputeShaderRHI, InstanceData->GPUExtentMinBuffer,   InstanceData->ExtentMinArray,   EPixelFormat::PF_R32_FLOAT,     TEXT("ExtentMinBuffer"));
-				SetBuffer(RHICmdList, ExtentMaxBuffer,   ComputeShaderRHI, InstanceData->GPUExtentMaxBuffer,   InstanceData->ExtentMaxArray,   EPixelFormat::PF_R32_FLOAT,     TEXT("ExtentMaxBuffer"));
-				SetBuffer(RHICmdList, VolumeBuffer,      ComputeShaderRHI, InstanceData->GPUVolumeBuffer,      InstanceData->VolumeArray,      EPixelFormat::PF_R32_FLOAT,     TEXT("VolumeBuffer"));
-				SetBuffer(RHICmdList, SolverIDBuffer,    ComputeShaderRHI, InstanceData->GPUSolverIDBuffer,    InstanceData->SolverIDArray,    EPixelFormat::PF_R32_SINT,      TEXT("SolverIDBuffer"));
-				SetBuffer(RHICmdList, DensityBuffer,     ComputeShaderRHI, InstanceData->GPUDensityBuffer,     InstanceData->DensityArray,     EPixelFormat::PF_R32_FLOAT,     TEXT("DensityBuffer"));
-				SetBuffer(RHICmdList, FrictionBuffer,    ComputeShaderRHI, InstanceData->GPUFrictionBuffer,    InstanceData->FrictionArray,    EPixelFormat::PF_R32_FLOAT,     TEXT("FrictionBuffer"));
-				SetBuffer(RHICmdList, RestitutionBuffer, ComputeShaderRHI, InstanceData->GPURestitutionBuffer, InstanceData->RestitutionArray, EPixelFormat::PF_R32_FLOAT,     TEXT("RestitutionBuffer"));
-				SetBuffer(RHICmdList, SurfaceTypeBuffer, ComputeShaderRHI, InstanceData->GPUSurfaceTypeBuffer, InstanceData->SurfaceTypeArray, EPixelFormat::PF_R32_SINT,      TEXT("SurfaceTypeBuffer"));
-				SetBuffer(RHICmdList, ColorBuffer,       ComputeShaderRHI, InstanceData->GPUColorBuffer,       InstanceData->ColorArray,       EPixelFormat::PF_A32B32G32R32F, TEXT("ColorBuffer"));
-
-				SetBuffer(RHICmdList, IncomingLocationBuffer,           ComputeShaderRHI, InstanceData->GPUIncomingLocationBuffer,           InstanceData->IncomingLocationArray,           EPixelFormat::PF_A32B32G32R32F, TEXT("IncomingLocationBuffer"));
-				SetBuffer(RHICmdList, IncomingAccumulatedImpulseBuffer, ComputeShaderRHI, InstanceData->GPUIncomingAccumulatedImpulseBuffer, InstanceData->IncomingAccumulatedImpulseArray, EPixelFormat::PF_A32B32G32R32F, TEXT("IncomingAccumulatedImpulseBuffer"));
-				SetBuffer(RHICmdList, IncomingNormalBuffer,             ComputeShaderRHI, InstanceData->GPUIncomingNormalBuffer,             InstanceData->IncomingNormalArray,             EPixelFormat::PF_A32B32G32R32F, TEXT("IncomingNormalBuffer"));
-				SetBuffer(RHICmdList, IncomingVelocity1Buffer,          ComputeShaderRHI, InstanceData->GPUIncomingVelocity1Buffer,          InstanceData->IncomingVelocity1Array,          EPixelFormat::PF_A32B32G32R32F, TEXT("IncomingVelocity1Buffer"));
-				SetBuffer(RHICmdList, IncomingVelocity2Buffer,          ComputeShaderRHI, InstanceData->GPUIncomingVelocity2Buffer,          InstanceData->IncomingVelocity2Array,          EPixelFormat::PF_A32B32G32R32F, TEXT("IncomingVelocity2Buffer"));
-				SetBuffer(RHICmdList, IncomingAngularVelocity1Buffer,   ComputeShaderRHI, InstanceData->GPUIncomingAngularVelocity1Buffer,   InstanceData->IncomingAngularVelocity1Array,   EPixelFormat::PF_A32B32G32R32F, TEXT("IncomingAngularVelocity1Buffer"));
-				SetBuffer(RHICmdList, IncomingAngularVelocity2Buffer,   ComputeShaderRHI, InstanceData->GPUIncomingAngularVelocity2Buffer,   InstanceData->IncomingAngularVelocity2Array,   EPixelFormat::PF_A32B32G32R32F, TEXT("IncomingAngularVelocity2Buffer"));
-				SetBuffer(RHICmdList, IncomingMass1Buffer,              ComputeShaderRHI, InstanceData->GPUIncomingMass1Buffer,              InstanceData->IncomingMass1Array,              EPixelFormat::PF_R32_FLOAT,     TEXT("IncomingMass1Buffer"));
-				SetBuffer(RHICmdList, IncomingMass2Buffer,              ComputeShaderRHI, InstanceData->GPUIncomingMass2Buffer,              InstanceData->IncomingMass2Array,              EPixelFormat::PF_R32_FLOAT,     TEXT("IncomingMass2Buffer"));
-				SetBuffer(RHICmdList, IncomingTimeBuffer,               ComputeShaderRHI, InstanceData->GPUIncomingTimeBuffer,               InstanceData->IncomingTimeArray,               EPixelFormat::PF_R32_FLOAT,     TEXT("IncomingTimeBuffer"));
-				
-				SetBuffer(RHICmdList, TransformTranslationBuffer,       ComputeShaderRHI, InstanceData->GPUTransformTranslationBuffer,       InstanceData->TransformTranslationArray,       EPixelFormat::PF_A32B32G32R32F, TEXT("TransformTranslationBuffer"));
-				SetBuffer(RHICmdList, TransformRotationBuffer,          ComputeShaderRHI, InstanceData->GPUTransformRotationBuffer,          InstanceData->TransformRotationArray,          EPixelFormat::PF_A32B32G32R32F, TEXT("TransformRotationBuffer"));
-				SetBuffer(RHICmdList, TransformScaleBuffer,             ComputeShaderRHI, InstanceData->GPUTransformScaleBuffer,             InstanceData->TransformScaleArray,             EPixelFormat::PF_A32B32G32R32F, TEXT("TransformScaleBuffer"));
-				SetBuffer(RHICmdList, BoundsBuffer,                     ComputeShaderRHI, InstanceData->GPUBoundsBuffer,                     InstanceData->BoundsArray,                     EPixelFormat::PF_A32B32G32R32F, TEXT("BoundsBuffer"));
-
-				SetShaderValue(RHICmdList, ComputeShaderRHI, LastSpawnedPointID, InstanceData->LastSpawnedPointID);
-				SetShaderValue(RHICmdList, ComputeShaderRHI, SolverTime, InstanceData->SolverTime);
-			}
+			FNiagaraDIChaosDestruction_GPUData& InstanceData = it.Value();
+			InstanceData.RDGPositionSRV = nullptr;
+			InstanceData.RDGVelocitySRV = nullptr;
+			InstanceData.RDGExtentMinBufferSRV = nullptr;
+			InstanceData.RDGExtentMaxBufferSRV = nullptr;
+			InstanceData.RDGVolumeBufferSRV = nullptr;
+			InstanceData.RDGSolverIDBufferSRV = nullptr;
+			InstanceData.RDGDensityBufferSRV = nullptr;
+			InstanceData.RDGFrictionBufferSRV = nullptr;
+			InstanceData.RDGRestitutionBufferSRV = nullptr;
+			InstanceData.RDGSurfaceTypeBufferSRV = nullptr;
+			InstanceData.RDGColorBufferSRV = nullptr;
+			InstanceData.RDGIncomingLocationBufferSRV = nullptr;
+			InstanceData.RDGIncomingAccumulatedImpulseBufferSRV = nullptr;
+			InstanceData.RDGIncomingNormalBufferSRV = nullptr;
+			InstanceData.RDGIncomingVelocity1BufferSRV = nullptr;
+			InstanceData.RDGIncomingVelocity2BufferSRV = nullptr;
+			InstanceData.RDGIncomingAngularVelocity1BufferSRV = nullptr;
+			InstanceData.RDGIncomingAngularVelocity2BufferSRV = nullptr;
+			InstanceData.RDGIncomingMass1BufferSRV = nullptr;
+			InstanceData.RDGIncomingMass2BufferSRV = nullptr;
+			InstanceData.RDGIncomingTimeBufferSRV = nullptr;
+			InstanceData.RDGTransformTranslationBufferSRV = nullptr;
+			InstanceData.RDGTransformRotationBufferSRV = nullptr;
+			InstanceData.RDGTransformScaleBufferSRV = nullptr;
+			InstanceData.RDGBoundsBufferSRV = nullptr;
 		}
 	}
-
-private:
-	// TODO: Collect these into a small number of buffers to reduce the number of binding points
-	LAYOUT_FIELD(FShaderResourceParameter, PositionBuffer);
-	LAYOUT_FIELD(FShaderResourceParameter, VelocityBuffer);
-	LAYOUT_FIELD(FShaderResourceParameter, ExtentMinBuffer);
-	LAYOUT_FIELD(FShaderResourceParameter, ExtentMaxBuffer);
-	LAYOUT_FIELD(FShaderResourceParameter, VolumeBuffer);
-	LAYOUT_FIELD(FShaderResourceParameter, SolverIDBuffer);
-	LAYOUT_FIELD(FShaderResourceParameter, DensityBuffer);
-	LAYOUT_FIELD(FShaderResourceParameter, FrictionBuffer);
-	LAYOUT_FIELD(FShaderResourceParameter, RestitutionBuffer);
-	LAYOUT_FIELD(FShaderResourceParameter, SurfaceTypeBuffer);
-	LAYOUT_FIELD(FShaderResourceParameter, ColorBuffer);
-
-	LAYOUT_FIELD(FShaderResourceParameter, IncomingLocationBuffer);
-	LAYOUT_FIELD(FShaderResourceParameter, IncomingAccumulatedImpulseBuffer);
-	LAYOUT_FIELD(FShaderResourceParameter, IncomingNormalBuffer);
-	LAYOUT_FIELD(FShaderResourceParameter, IncomingVelocity1Buffer);
-	LAYOUT_FIELD(FShaderResourceParameter, IncomingVelocity2Buffer);
-	LAYOUT_FIELD(FShaderResourceParameter, IncomingAngularVelocity1Buffer);
-	LAYOUT_FIELD(FShaderResourceParameter, IncomingAngularVelocity2Buffer);
-	LAYOUT_FIELD(FShaderResourceParameter, IncomingMass1Buffer);
-	LAYOUT_FIELD(FShaderResourceParameter, IncomingMass2Buffer);
-	LAYOUT_FIELD(FShaderResourceParameter, IncomingTimeBuffer);
-
-	LAYOUT_FIELD(FShaderResourceParameter, TransformTranslationBuffer);
-	LAYOUT_FIELD(FShaderResourceParameter, TransformRotationBuffer);
-	LAYOUT_FIELD(FShaderResourceParameter, TransformScaleBuffer);
-	LAYOUT_FIELD(FShaderResourceParameter, BoundsBuffer);
-
-	LAYOUT_FIELD(FShaderParameter, LastSpawnedPointID);
-	LAYOUT_FIELD(FShaderParameter, SolverTime);
-};
-
-IMPLEMENT_TYPE_LAYOUT(FNiagaraDataInterfaceParametersCS_ChaosDestruction);
-
-IMPLEMENT_NIAGARA_DI_PARAMETER(UNiagaraDataInterfaceChaosDestruction, FNiagaraDataInterfaceParametersCS_ChaosDestruction);
-
-
-//#pragma optimize("", on)
-
-
+}
 
 #undef LOCTEXT_NAMESPACE
diff --git a/Engine/Plugins/Experimental/ChaosNiagara/Source/ChaosNiagara/Private/NiagaraDataInterfaceGeometryCollection.cpp b/Engine/Plugins/Experimental/ChaosNiagara/Source/ChaosNiagara/Private/NiagaraDataInterfaceGeometryCollection.cpp
--- a/Engine/Plugins/Experimental/ChaosNiagara/Source/ChaosNiagara/Private/NiagaraDataInterfaceGeometryCollection.cpp
+++ b/Engine/Plugins/Experimental/ChaosNiagara/Source/ChaosNiagara/Private/NiagaraDataInterfaceGeometryCollection.cpp
@@ -3,6 +3,7 @@
 #include "NiagaraDataInterfaceGeometryCollection.h"
 #include "NiagaraRenderer.h"
 #include "NiagaraSimStageData.h"
+#include "NiagaraShaderParametersBuilder.h"
 #include "NiagaraSystemInstance.h"
 #include "ShaderParameterUtils.h"
 #include "GeometryCollection/GeometryCollectionComponent.h"
@@ -12,86 +13,47 @@ DEFINE_LOG_CATEGORY_STATIC(LogGeometryCollection, Log, All);
 
 //------------------------------------------------------------------------------------------------------------
 
-static const FName GetClosestPointNoNormalName(TEXT("GetClosestPointNoNormal"));
-
-
-//------------------------------------------------------------------------------------------------------------
-
-const FString UNiagaraDataInterfaceGeometryCollection::BoundsMinName(TEXT("BoundsMin_"));
-const FString UNiagaraDataInterfaceGeometryCollection::BoundsMaxName(TEXT("BoundsMax_"));
-const FString UNiagaraDataInterfaceGeometryCollection::NumPiecesName(TEXT("NumPieces_"));
-const FString UNiagaraDataInterfaceGeometryCollection::WorldTransformBufferName(TEXT("WorldTransformBuffer_"));
-const FString UNiagaraDataInterfaceGeometryCollection::PrevWorldTransformBufferName(TEXT("PrevWorldTransformBuffer_"));
-const FString UNiagaraDataInterfaceGeometryCollection::WorldInverseTransformBufferName(TEXT("WorldInverseTransformBuffer_"));
-const FString UNiagaraDataInterfaceGeometryCollection::PrevWorldInverseTransformBufferName(TEXT("PrevWorldInverseTransformBuffer_"));
-const FString UNiagaraDataInterfaceGeometryCollection::BoundsBufferName(TEXT("BoundsBuffer_"));
-
-
-//------------------------------------------------------------------------------------------------------------
-
-struct FNDIGeometryCollectionParametersName
+namespace NDIGeometryCollectionLocal
 {
-	FNDIGeometryCollectionParametersName(const FString& Suffix)
-	{
-		BoundsMinName = UNiagaraDataInterfaceGeometryCollection::BoundsMinName + Suffix;
-		BoundsMaxName = UNiagaraDataInterfaceGeometryCollection::BoundsMaxName + Suffix;
-		NumPiecesName = UNiagaraDataInterfaceGeometryCollection::NumPiecesName + Suffix;
-		WorldTransformBufferName = UNiagaraDataInterfaceGeometryCollection::WorldTransformBufferName + Suffix;
-		PrevWorldTransformBufferName = UNiagaraDataInterfaceGeometryCollection::PrevWorldTransformBufferName + Suffix;
-		WorldInverseTransformBufferName = UNiagaraDataInterfaceGeometryCollection::WorldInverseTransformBufferName + Suffix;
-		PrevWorldInverseTransformBufferName = UNiagaraDataInterfaceGeometryCollection::PrevWorldInverseTransformBufferName + Suffix;
-		BoundsBufferName = UNiagaraDataInterfaceGeometryCollection::BoundsBufferName + Suffix;
-	}
-
-	FString BoundsMinName;
-	FString BoundsMaxName;
+	static const FName GetClosestPointNoNormalName(TEXT("GetClosestPointNoNormal"));
+	static const TCHAR* TemplateShaderFilePath = TEXT("/Plugin/Experimental/ChaosNiagara/NiagaraDataInterfaceGeometryCollection.ush");
 
-	FString NumPiecesName;
-	FString WorldTransformBufferName;
-	FString PrevWorldTransformBufferName;
-	FString WorldInverseTransformBufferName;
-	FString PrevWorldInverseTransformBufferName;
-	FString BoundsBufferName;
-};
-
-//------------------------------------------------------------------------------------------------------------
-
-template<typename BufferType, EPixelFormat PixelFormat>
-void CreateInternalBuffer(FRWBuffer& OutputBuffer, uint32 ElementCount)
-{
-	if (ElementCount > 0)
+	template<typename BufferType, EPixelFormat PixelFormat>
+	void CreateInternalBuffer(FReadBuffer& OutputBuffer, uint32 ElementCount)
 	{
-		OutputBuffer.Initialize(TEXT("FNDIGeometryCollectionBuffer"), sizeof(BufferType), ElementCount, PixelFormat, BUF_Static);
+		if (ElementCount > 0)
+		{
+			OutputBuffer.Initialize(TEXT("FNDIGeometryCollectionBuffer"), sizeof(BufferType), ElementCount, PixelFormat, BUF_Static);
+		}
 	}
-}
 
-template<typename BufferType, EPixelFormat PixelFormat>
-void UpdateInternalBuffer(const TArray<BufferType>& InputData, FRWBuffer& OutputBuffer)
-{
-	uint32 ElementCount = InputData.Num();
-	if (ElementCount > 0 && OutputBuffer.Buffer.IsValid())
+	template<typename BufferType, EPixelFormat PixelFormat>
+	void UpdateInternalBuffer(const TArray<BufferType>& InputData, FReadBuffer& OutputBuffer)
 	{
-		const uint32 BufferBytes = sizeof(BufferType) * ElementCount;
+		uint32 ElementCount = InputData.Num();
+		if (ElementCount > 0 && OutputBuffer.Buffer.IsValid())
+		{
+			const uint32 BufferBytes = sizeof(BufferType) * ElementCount;
 
-		void* OutputData = RHILockBuffer(OutputBuffer.Buffer, 0, BufferBytes, RLM_WriteOnly);
+			void* OutputData = RHILockBuffer(OutputBuffer.Buffer, 0, BufferBytes, RLM_WriteOnly);
 
-		FMemory::Memcpy(OutputData, InputData.GetData(), BufferBytes);
-		RHIUnlockBuffer(OutputBuffer.Buffer);
+			FMemory::Memcpy(OutputData, InputData.GetData(), BufferBytes);
+			RHIUnlockBuffer(OutputBuffer.Buffer);
+		}
 	}
 }
 
-
 //------------------------------------------------------------------------------------------------------------
 
 void FNDIGeometryCollectionBuffer::InitRHI()
 {
-	CreateInternalBuffer<FVector4f, EPixelFormat::PF_A32B32G32R32F>(WorldTransformBuffer, 3 * NumPieces);
-	CreateInternalBuffer<FVector4f, EPixelFormat::PF_A32B32G32R32F>(PrevWorldTransformBuffer, 3 * NumPieces);
+	NDIGeometryCollectionLocal::CreateInternalBuffer<FVector4f, EPixelFormat::PF_A32B32G32R32F>(WorldTransformBuffer, 3 * NumPieces);
+	NDIGeometryCollectionLocal::CreateInternalBuffer<FVector4f, EPixelFormat::PF_A32B32G32R32F>(PrevWorldTransformBuffer, 3 * NumPieces);
 
-	CreateInternalBuffer<FVector4f, EPixelFormat::PF_A32B32G32R32F>(WorldInverseTransformBuffer, 3 * NumPieces);
-	CreateInternalBuffer<FVector4f, EPixelFormat::PF_A32B32G32R32F>(PrevWorldInverseTransformBuffer, 3 * NumPieces);
+	NDIGeometryCollectionLocal::CreateInternalBuffer<FVector4f, EPixelFormat::PF_A32B32G32R32F>(WorldInverseTransformBuffer, 3 * NumPieces);
+	NDIGeometryCollectionLocal::CreateInternalBuffer<FVector4f, EPixelFormat::PF_A32B32G32R32F>(PrevWorldInverseTransformBuffer, 3 * NumPieces);
 
-	CreateInternalBuffer<FVector4f, EPixelFormat::PF_A32B32G32R32F>(BoundsBuffer, NumPieces);
+	NDIGeometryCollectionLocal::CreateInternalBuffer<FVector4f, EPixelFormat::PF_A32B32G32R32F>(BoundsBuffer, NumPieces);
 }
 
 void FNDIGeometryCollectionBuffer::ReleaseRHI()
@@ -285,107 +247,13 @@ ETickingGroup FNDIGeometryCollectionData::ComputeTickingGroup()
 	return TickingGroup;
 }
 
-
-//------------------------------------------------------------------------------------------------------------
-
-struct FNDIGeometryCollectionParametersCS : public FNiagaraDataInterfaceParametersCS
-{
-	DECLARE_TYPE_LAYOUT(FNDIGeometryCollectionParametersCS, NonVirtual);
-public:
-	void Bind(const FNiagaraDataInterfaceGPUParamInfo& ParameterInfo, const class FShaderParameterMap& ParameterMap)
-	{
-		FNDIGeometryCollectionParametersName ParamNames(*ParameterInfo.DataInterfaceHLSLSymbol);
-
-		BoundsMin.Bind(ParameterMap, *ParamNames.BoundsMinName);
-		BoundsMax.Bind(ParameterMap, *ParamNames.BoundsMaxName);
-
-		NumPieces.Bind(ParameterMap, *ParamNames.NumPiecesName);
-		WorldTransformBuffer.Bind(ParameterMap, *ParamNames.WorldTransformBufferName);
-		PrevWorldTransformBuffer.Bind(ParameterMap, *ParamNames.PrevWorldTransformBufferName);
-		WorldInverseTransformBuffer.Bind(ParameterMap, *ParamNames.WorldInverseTransformBufferName);
-		PrevWorldInverseTransformBuffer.Bind(ParameterMap, *ParamNames.PrevWorldInverseTransformBufferName);
-		BoundsBuffer.Bind(ParameterMap, *ParamNames.BoundsBufferName);
-	}
-
-	void Set(FRHICommandList& RHICmdList, const FNiagaraDataInterfaceSetArgs& Context) const
-	{
-		check(IsInRenderingThread());
-
-		FRHIComputeShader* ComputeShaderRHI = RHICmdList.GetBoundComputeShader();
-
-		FNDIGeometryCollectionProxy* InterfaceProxy =
-			static_cast<FNDIGeometryCollectionProxy*>(Context.DataInterface);
-		FNDIGeometryCollectionData* ProxyData =
-			InterfaceProxy->SystemInstancesToProxyData.Find(Context.SystemInstanceID);
-
-		if (ProxyData != nullptr)
-		{
-
-			FNDIGeometryCollectionBuffer* AssetBuffer = ProxyData->AssetBuffer;
-
-			FRHITransitionInfo Transitions[] = {
-				FRHITransitionInfo(AssetBuffer->WorldTransformBuffer.UAV, ERHIAccess::Unknown, ERHIAccess::SRVCompute),
-				FRHITransitionInfo(AssetBuffer->PrevWorldTransformBuffer.UAV, ERHIAccess::Unknown, ERHIAccess::SRVCompute),
-				FRHITransitionInfo(AssetBuffer->WorldInverseTransformBuffer.UAV, ERHIAccess::Unknown, ERHIAccess::SRVCompute),
-				FRHITransitionInfo(AssetBuffer->PrevWorldInverseTransformBuffer.UAV, ERHIAccess::Unknown, ERHIAccess::SRVCompute),
-				FRHITransitionInfo(AssetBuffer->BoundsBuffer.UAV, ERHIAccess::Unknown, ERHIAccess::SRVCompute),
-			};
-			RHICmdList.Transition(MakeArrayView(Transitions, UE_ARRAY_COUNT(Transitions)));		
-			
-			SetShaderValue(RHICmdList, ComputeShaderRHI, BoundsMin, ProxyData->BoundsOrigin - ProxyData->BoundsExtent);
-			SetShaderValue(RHICmdList, ComputeShaderRHI, BoundsMax, ProxyData->BoundsOrigin + ProxyData->BoundsExtent);
-
-			SetShaderValue(RHICmdList, ComputeShaderRHI, NumPieces, ProxyData->AssetBuffer->NumPieces);
-			SetSRVParameter(RHICmdList, ComputeShaderRHI, WorldTransformBuffer, ProxyData->AssetBuffer->WorldTransformBuffer.SRV);
-			SetSRVParameter(RHICmdList, ComputeShaderRHI, PrevWorldTransformBuffer, ProxyData->AssetBuffer->PrevWorldTransformBuffer.SRV);
-			SetSRVParameter(RHICmdList, ComputeShaderRHI, WorldInverseTransformBuffer, ProxyData->AssetBuffer->WorldInverseTransformBuffer.SRV);
-			SetSRVParameter(RHICmdList, ComputeShaderRHI, PrevWorldInverseTransformBuffer, ProxyData->AssetBuffer->PrevWorldInverseTransformBuffer.SRV);
-			SetSRVParameter(RHICmdList, ComputeShaderRHI, BoundsBuffer, ProxyData->AssetBuffer->BoundsBuffer.SRV);
-		}
-		else
-		{
-			SetShaderValue(RHICmdList, ComputeShaderRHI, BoundsMin, FVector3f(0, 0, 0));
-			SetShaderValue(RHICmdList, ComputeShaderRHI, BoundsMax, FVector3f(0, 0, 0));
-
-			SetShaderValue(RHICmdList, ComputeShaderRHI, NumPieces, 0);
-			SetSRVParameter(RHICmdList, ComputeShaderRHI, WorldTransformBuffer, FNiagaraRenderer::GetDummyFloat4Buffer());
-			SetSRVParameter(RHICmdList, ComputeShaderRHI, PrevWorldTransformBuffer, FNiagaraRenderer::GetDummyFloat4Buffer());
-			SetSRVParameter(RHICmdList, ComputeShaderRHI, WorldInverseTransformBuffer, FNiagaraRenderer::GetDummyFloat4Buffer());
-			SetSRVParameter(RHICmdList, ComputeShaderRHI, PrevWorldInverseTransformBuffer, FNiagaraRenderer::GetDummyFloat4Buffer());
-			SetSRVParameter(RHICmdList, ComputeShaderRHI, BoundsBuffer, FNiagaraRenderer::GetDummyFloat4Buffer());
-		}
-	}
-
-	void Unset(FRHICommandList& RHICmdList, const FNiagaraDataInterfaceSetArgs& Context) const
-	{
-	}
-
-private:
-
-	LAYOUT_FIELD(FShaderParameter, BoundsMin);
-	LAYOUT_FIELD(FShaderParameter, BoundsMax);
-
-	LAYOUT_FIELD(FShaderParameter, NumPieces);
-
-	LAYOUT_FIELD(FShaderResourceParameter, WorldTransformBuffer);
-	LAYOUT_FIELD(FShaderResourceParameter, PrevWorldTransformBuffer);
-	LAYOUT_FIELD(FShaderResourceParameter, WorldInverseTransformBuffer);
-	LAYOUT_FIELD(FShaderResourceParameter, PrevWorldInverseTransformBuffer);
-	LAYOUT_FIELD(FShaderResourceParameter, BoundsBuffer);
-};
-
-IMPLEMENT_TYPE_LAYOUT(FNDIGeometryCollectionParametersCS);
-
-IMPLEMENT_NIAGARA_DI_PARAMETER(UNiagaraDataInterfaceGeometryCollection, FNDIGeometryCollectionParametersCS);
-
-
 //------------------------------------------------------------------------------------------------------------
 
-void FNDIGeometryCollectionProxy::ConsumePerInstanceDataFromGameThread(void* PerInstanceData, const FNiagaraSystemInstanceID& Instance)
+void FNDIGeometryCollectionProxy::ConsumePerInstanceDataFromGameThread(void* DataFromGameThread, const FNiagaraSystemInstanceID& Instance)
 {
 	check(IsInRenderingThread());
 
-	FNDIGeometryCollectionData* SourceData = static_cast<FNDIGeometryCollectionData*>(PerInstanceData);	
+	FNDIGeometryCollectionData* SourceData = static_cast<FNDIGeometryCollectionData*>(DataFromGameThread);
 	FNDIGeometryCollectionData* TargetData = &(SystemInstancesToProxyData.FindOrAdd(Instance));
 
 	ensure(TargetData);
@@ -401,6 +269,7 @@ void FNDIGeometryCollectionProxy::ConsumePerInstanceDataFromGameThread(void* Per
 	{
 		UE_LOG(LogGeometryCollection, Log, TEXT("ConsumePerInstanceDataFromGameThread() ... could not find %d"), Instance);
 	}
+	SourceData->~FNDIGeometryCollectionData();
 }
 
 void FNDIGeometryCollectionProxy::InitializePerInstanceData(const FNiagaraSystemInstanceID& SystemInstance)
@@ -418,34 +287,23 @@ void FNDIGeometryCollectionProxy::DestroyPerInstanceData(const FNiagaraSystemIns
 	SystemInstancesToProxyData.Remove(SystemInstance);
 }
 
-void FNDIGeometryCollectionProxy::PreStage(FRHICommandList& RHICmdList, const FNiagaraDataInterfaceStageArgs& Context)
+void FNDIGeometryCollectionProxy::PreStage(const FNDIGpuComputePreStageContext& Context)
 {
-	check(SystemInstancesToProxyData.Contains(Context.SystemInstanceID));
-
-	FNDIGeometryCollectionData* ProxyData =
-		SystemInstancesToProxyData.Find(Context.SystemInstanceID);
+	check(SystemInstancesToProxyData.Contains(Context.GetSystemInstanceID()));
 
+	FNDIGeometryCollectionData* ProxyData = SystemInstancesToProxyData.Find(Context.GetSystemInstanceID());
 	if (ProxyData != nullptr && ProxyData->AssetBuffer)
 	{
-		if (Context.SimStageData->bFirstStage)
+		if (Context.GetSimStageData().bFirstStage)
 		{
 			FNDIGeometryCollectionBuffer* AssetBuffer = ProxyData->AssetBuffer;
 
-			FRHITransitionInfo Transitions[] = {
-				FRHITransitionInfo(AssetBuffer->WorldTransformBuffer.UAV, ERHIAccess::Unknown, ERHIAccess::UAVCompute),
-				FRHITransitionInfo(AssetBuffer->PrevWorldTransformBuffer.UAV, ERHIAccess::Unknown, ERHIAccess::UAVCompute),
-				FRHITransitionInfo(AssetBuffer->WorldInverseTransformBuffer.UAV, ERHIAccess::Unknown, ERHIAccess::UAVCompute),
-				FRHITransitionInfo(AssetBuffer->PrevWorldInverseTransformBuffer.UAV, ERHIAccess::Unknown, ERHIAccess::UAVCompute),
-				FRHITransitionInfo(AssetBuffer->BoundsBuffer.UAV, ERHIAccess::Unknown, ERHIAccess::UAVCompute),
-			};
-			RHICmdList.Transition(MakeArrayView(Transitions, UE_ARRAY_COUNT(Transitions)));
-
 			// #todo(dmp): bounds buffer doesn't need to be updated each frame
-			UpdateInternalBuffer<FVector4f, EPixelFormat::PF_A32B32G32R32F>(ProxyData->AssetArrays->WorldTransformBuffer, ProxyData->AssetBuffer->WorldTransformBuffer);
-			UpdateInternalBuffer<FVector4f, EPixelFormat::PF_A32B32G32R32F>(ProxyData->AssetArrays->PrevWorldTransformBuffer, ProxyData->AssetBuffer->PrevWorldTransformBuffer);
-			UpdateInternalBuffer<FVector4f, EPixelFormat::PF_A32B32G32R32F>(ProxyData->AssetArrays->WorldInverseTransformBuffer, ProxyData->AssetBuffer->WorldInverseTransformBuffer);
-			UpdateInternalBuffer<FVector4f, EPixelFormat::PF_A32B32G32R32F>(ProxyData->AssetArrays->PrevWorldInverseTransformBuffer, ProxyData->AssetBuffer->PrevWorldInverseTransformBuffer);
-			UpdateInternalBuffer<FVector4f, EPixelFormat::PF_A32B32G32R32F>(ProxyData->AssetArrays->BoundsBuffer, ProxyData->AssetBuffer->BoundsBuffer);
+			NDIGeometryCollectionLocal::UpdateInternalBuffer<FVector4f, EPixelFormat::PF_A32B32G32R32F>(ProxyData->AssetArrays->WorldTransformBuffer, ProxyData->AssetBuffer->WorldTransformBuffer);
+			NDIGeometryCollectionLocal::UpdateInternalBuffer<FVector4f, EPixelFormat::PF_A32B32G32R32F>(ProxyData->AssetArrays->PrevWorldTransformBuffer, ProxyData->AssetBuffer->PrevWorldTransformBuffer);
+			NDIGeometryCollectionLocal::UpdateInternalBuffer<FVector4f, EPixelFormat::PF_A32B32G32R32F>(ProxyData->AssetArrays->WorldInverseTransformBuffer, ProxyData->AssetBuffer->WorldInverseTransformBuffer);
+			NDIGeometryCollectionLocal::UpdateInternalBuffer<FVector4f, EPixelFormat::PF_A32B32G32R32F>(ProxyData->AssetArrays->PrevWorldInverseTransformBuffer, ProxyData->AssetBuffer->PrevWorldInverseTransformBuffer);
+			NDIGeometryCollectionLocal::UpdateInternalBuffer<FVector4f, EPixelFormat::PF_A32B32G32R32F>(ProxyData->AssetArrays->BoundsBuffer, ProxyData->AssetBuffer->BoundsBuffer);
 		}
 	}
 }
@@ -489,16 +347,16 @@ void UNiagaraDataInterfaceGeometryCollection::DestroyPerInstanceData(void* PerIn
 	FNDIGeometryCollectionProxy* ThisProxy = GetProxyAs<FNDIGeometryCollectionProxy>();
 	ENQUEUE_RENDER_COMMAND(FNiagaraDIDestroyInstanceData) (
 		[ThisProxy, InstanceID = SystemInstance->GetId()](FRHICommandListImmediate& CmdList)
-	{
-		FNDIGeometryCollectionData* ProxyData =
-			ThisProxy->SystemInstancesToProxyData.Find(InstanceID);
-
-		if (ProxyData != nullptr && ProxyData->AssetArrays)
-		{			
-			ThisProxy->SystemInstancesToProxyData.Remove(InstanceID);
-			delete ProxyData->AssetArrays;
-		}		
-	}
+		{
+			FNDIGeometryCollectionData* ProxyData =
+				ThisProxy->SystemInstancesToProxyData.Find(InstanceID);
+
+			if (ProxyData != nullptr && ProxyData->AssetArrays)
+			{			
+				ThisProxy->SystemInstancesToProxyData.Remove(InstanceID);
+				delete ProxyData->AssetArrays;
+			}		
+		}
 	);
 }
 
@@ -552,7 +410,7 @@ void UNiagaraDataInterfaceGeometryCollection::GetFunctions(TArray<FNiagaraFuncti
 {
 	{
 		FNiagaraFunctionSignature Sig;
-		Sig.Name = GetClosestPointNoNormalName;
+		Sig.Name = NDIGeometryCollectionLocal::GetClosestPointNoNormalName;
 		Sig.bSupportsGPU = true;
 		Sig.bSupportsCPU = false;
 		Sig.bMemberFunction = true;
@@ -576,53 +434,66 @@ void UNiagaraDataInterfaceGeometryCollection::GetVMExternalFunction(const FVMExt
 #if WITH_EDITORONLY_DATA
 bool UNiagaraDataInterfaceGeometryCollection::GetFunctionHLSL(const FNiagaraDataInterfaceGPUParamInfo& ParamInfo, const FNiagaraDataInterfaceGeneratedFunction& FunctionInfo, int FunctionInstanceIndex, FString& OutHLSL)
 {
-	FNDIGeometryCollectionParametersName ParamNames(ParamInfo.DataInterfaceHLSLSymbol);
-
-	TMap<FString, FStringFormatArg> ArgsSample = {
-
-		{TEXT("InstanceFunctionName"), FunctionInfo.InstanceName},
-		{TEXT("NumPiecesName"), ParamNames.NumPiecesName},
-		{TEXT("WorldTransformBufferName"), ParamNames.WorldTransformBufferName},
-		{TEXT("PrevWorldTransformBufferName"), ParamNames.PrevWorldTransformBufferName},
-		{TEXT("WorldInverseTransformBufferName"), ParamNames.WorldInverseTransformBufferName},
-		{TEXT("PrevWorldInverseTransformBufferName"), ParamNames.PrevWorldInverseTransformBufferName},
-		{TEXT("BoundsBufferName"), ParamNames.BoundsBufferName},
-		{TEXT("GeometryCollectionContextName"), TEXT("DIGEOMETRYCOLLECTION_MAKE_CONTEXT(") + ParamInfo.DataInterfaceHLSLSymbol + TEXT(")")},
-	};
-
-	if (FunctionInfo.DefinitionName == GetClosestPointNoNormalName)
+	if (FunctionInfo.DefinitionName == NDIGeometryCollectionLocal::GetClosestPointNoNormalName)
 	{
-	static const TCHAR* FormatSample = TEXT(R"(
-		void {InstanceFunctionName}(in float3 WorldPosition, in float DeltaTime, in float TimeFraction, out float ClosestDistance, out float3 OutClosestPosition, 
-							out float3 OutClosestVelocity)
-		{
-			{GeometryCollectionContextName} DIGeometryCollection_GetClosestPointNoNormal(DIContext,WorldPosition,DeltaTime,TimeFraction, ClosestDistance,
-				OutClosestPosition,OutClosestVelocity);
-		}
-		)");
-	OutHLSL += FString::Format(FormatSample, ArgsSample);
-	return true;
+		return true;
 	}
-	OutHLSL += TEXT("\n");
+
 	return false;
 }
 
-void UNiagaraDataInterfaceGeometryCollection::GetCommonHLSL(FString& OutHLSL)
-{	
-	OutHLSL += TEXT("#include \"/Plugin/Experimental/ChaosNiagara/NiagaraDataInterfaceGeometryCollection.ush\"\n");			
+bool UNiagaraDataInterfaceGeometryCollection::AppendCompileHash(FNiagaraCompileHashVisitor* InVisitor) const
+{
+	bool bSuccess = Super::AppendCompileHash(InVisitor);
+	bSuccess &= InVisitor->UpdateString(TEXT("UNiagaraDataInterfaceGeometryCollectionSource"), GetShaderFileHash(NDIGeometryCollectionLocal::TemplateShaderFilePath, EShaderPlatform::SP_PCD3D_SM5).ToString());
+	bSuccess &= InVisitor->UpdateShaderParameters<FShaderParameters>();
+	return bSuccess;
 }
 
 void UNiagaraDataInterfaceGeometryCollection::GetParameterDefinitionHLSL(const FNiagaraDataInterfaceGPUParamInfo& ParamInfo, FString& OutHLSL)
 {
-	OutHLSL += TEXT("DIGEOMETRYCOLLECTION_DECLARE_CONSTANTS(") + ParamInfo.DataInterfaceHLSLSymbol + TEXT(")\n");
+	const TMap<FString, FStringFormatArg> TemplateArgs = { {TEXT("ParameterName"),	ParamInfo.DataInterfaceHLSLSymbol}, };
+
+	FString TemplateFile;
+	LoadShaderSourceFile(NDIGeometryCollectionLocal::TemplateShaderFilePath, EShaderPlatform::SP_PCD3D_SM5, &TemplateFile, nullptr);
+	OutHLSL += FString::Format(*TemplateFile, TemplateArgs);
 }
 #endif
 
-#if WITH_EDITOR
-void UNiagaraDataInterfaceGeometryCollection::ValidateFunction(const FNiagaraFunctionSignature& Function, TArray<FText>& OutValidationErrors)
+void UNiagaraDataInterfaceGeometryCollection::BuildShaderParameters(FNiagaraShaderParametersBuilder& ShaderParametersBuilder) const
 {
+	ShaderParametersBuilder.AddNestedStruct<FShaderParameters>();
+}
+
+void UNiagaraDataInterfaceGeometryCollection::SetShaderParameters(const FNiagaraDataInterfaceSetShaderParametersContext& Context) const
+{
+	FNDIGeometryCollectionProxy& InterfaceProxy = Context.GetProxy<FNDIGeometryCollectionProxy>();
+	FNDIGeometryCollectionData* ProxyData = InterfaceProxy.SystemInstancesToProxyData.Find(Context.GetSystemInstanceID());
+
+	FShaderParameters* ShaderParameters = Context.GetParameterNestedStruct<FShaderParameters>();
+	if (ProxyData != nullptr)
+	{
+		ShaderParameters->BoundsMin							= ProxyData->BoundsOrigin - ProxyData->BoundsExtent;
+		ShaderParameters->BoundsMax							= ProxyData->BoundsOrigin + ProxyData->BoundsExtent;
+		ShaderParameters->NumPieces							= ProxyData->AssetBuffer->NumPieces;
+		ShaderParameters->WorldTransformBuffer				= FNiagaraRenderer::GetSrvOrDefaultFloat4(ProxyData->AssetBuffer->WorldTransformBuffer.SRV);
+		ShaderParameters->PrevWorldTransformBuffer			= FNiagaraRenderer::GetSrvOrDefaultFloat4(ProxyData->AssetBuffer->PrevWorldTransformBuffer.SRV);
+		ShaderParameters->WorldInverseTransformBuffer		= FNiagaraRenderer::GetSrvOrDefaultFloat4(ProxyData->AssetBuffer->WorldInverseTransformBuffer.SRV);
+		ShaderParameters->PrevWorldInverseTransformBuffer	= FNiagaraRenderer::GetSrvOrDefaultFloat4(ProxyData->AssetBuffer->PrevWorldInverseTransformBuffer.SRV);
+		ShaderParameters->BoundsBuffer						= FNiagaraRenderer::GetSrvOrDefaultFloat4(ProxyData->AssetBuffer->BoundsBuffer.SRV);
+	}
+	else
+	{
+		ShaderParameters->BoundsMin							= FVector3f::ZeroVector;
+		ShaderParameters->BoundsMax							= FVector3f::ZeroVector;
+		ShaderParameters->NumPieces							= 0;
+		ShaderParameters->WorldTransformBuffer				= FNiagaraRenderer::GetDummyFloat4Buffer();
+		ShaderParameters->PrevWorldTransformBuffer			= FNiagaraRenderer::GetDummyFloat4Buffer();
+		ShaderParameters->WorldInverseTransformBuffer		= FNiagaraRenderer::GetDummyFloat4Buffer();
+		ShaderParameters->PrevWorldInverseTransformBuffer	= FNiagaraRenderer::GetDummyFloat4Buffer();
+		ShaderParameters->BoundsBuffer						= FNiagaraRenderer::GetDummyFloat4Buffer();
+	}
 }
-#endif
 
 void UNiagaraDataInterfaceGeometryCollection::ProvidePerInstanceDataForRenderThread(void* DataForRenderThread, void* PerInstanceData, const FNiagaraSystemInstanceID& SystemInstance)
 {
@@ -642,4 +513,4 @@ void UNiagaraDataInterfaceGeometryCollection::ProvidePerInstanceDataForRenderThr
 	check(Proxy);
 }
 
-#undef LOCTEXT_NAMESPACE
\ No newline at end of file
+#undef LOCTEXT_NAMESPACE
diff --git a/Engine/Plugins/Experimental/ChaosNiagara/Source/ChaosNiagara/Public/NiagaraDataInterfaceGeometryCollection.h b/Engine/Plugins/Experimental/ChaosNiagara/Source/ChaosNiagara/Public/NiagaraDataInterfaceGeometryCollection.h
--- a/Engine/Plugins/Experimental/ChaosNiagara/Source/ChaosNiagara/Public/NiagaraDataInterfaceGeometryCollection.h
+++ b/Engine/Plugins/Experimental/ChaosNiagara/Source/ChaosNiagara/Public/NiagaraDataInterfaceGeometryCollection.h
@@ -67,19 +67,19 @@ struct FNDIGeometryCollectionBuffer : public FRenderResource
 	virtual FString GetFriendlyName() const override { return TEXT("FNDIGeometryCollectionBuffer"); }
 
 	/** World transform buffer */
-	FRWBuffer WorldTransformBuffer;
+	FReadBuffer WorldTransformBuffer;
 
 	/** Inverse transform buffer*/
-	FRWBuffer PrevWorldTransformBuffer;
+	FReadBuffer PrevWorldTransformBuffer;
 
 	/** World transform buffer */
-	FRWBuffer WorldInverseTransformBuffer;
+	FReadBuffer WorldInverseTransformBuffer;
 
 	/** Inverse transform buffer*/
-	FRWBuffer PrevWorldInverseTransformBuffer;
+	FReadBuffer PrevWorldInverseTransformBuffer;
 
 	/** Element extent buffer */
-	FRWBuffer BoundsBuffer;
+	FReadBuffer BoundsBuffer;
 
 	/** number of transforms */
 	uint32 NumPieces;
@@ -125,10 +125,18 @@ class UNiagaraDataInterfaceGeometryCollection : public UNiagaraDataInterface
 {
 	GENERATED_UCLASS_BODY()
 
-public:
-
-	DECLARE_NIAGARA_DI_PARAMETER();
+	BEGIN_SHADER_PARAMETER_STRUCT(FShaderParameters, )
+		SHADER_PARAMETER(FVector3f,				BoundsMin)
+		SHADER_PARAMETER(FVector3f,				BoundsMax)
+		SHADER_PARAMETER(uint32,				NumPieces)
+		SHADER_PARAMETER_SRV(Buffer<float4>,	WorldTransformBuffer)
+		SHADER_PARAMETER_SRV(Buffer<float4>,	PrevWorldTransformBuffer)
+		SHADER_PARAMETER_SRV(Buffer<float4>,	WorldInverseTransformBuffer)
+		SHADER_PARAMETER_SRV(Buffer<float4>,	PrevWorldInverseTransformBuffer)
+		SHADER_PARAMETER_SRV(Buffer<float4>,	BoundsBuffer)
+	END_SHADER_PARAMETER_STRUCT()
 
+public:
 	UPROPERTY(EditAnywhere, Category = "Geometry Collection")
 	TObjectPtr<AGeometryCollectionActor> GeometryCollectionActor;
 
@@ -151,22 +159,15 @@ public:
 
 	/** GPU simulation  functionality */
 #if WITH_EDITORONLY_DATA
-	virtual void GetCommonHLSL(FString& OutHLSL) override;
+	virtual bool AppendCompileHash(FNiagaraCompileHashVisitor* InVisitor) const override;
 	virtual void GetParameterDefinitionHLSL(const FNiagaraDataInterfaceGPUParamInfo& ParamInfo, FString& OutHLSL) override;
 	virtual bool GetFunctionHLSL(const FNiagaraDataInterfaceGPUParamInfo& ParamInfo, const FNiagaraDataInterfaceGeneratedFunction& FunctionInfo, int FunctionInstanceIndex, FString& OutHLSL) override;
-
-	virtual void ValidateFunction(const FNiagaraFunctionSignature& Function, TArray<FText>& OutValidationErrors) override;
 #endif
-	virtual void ProvidePerInstanceDataForRenderThread(void* DataForRenderThread, void* PerInstanceData, const FNiagaraSystemInstanceID& SystemInstance) override;
+	virtual bool UseLegacyShaderBindings() const override { return false; }
+	virtual void BuildShaderParameters(FNiagaraShaderParametersBuilder& ShaderParametersBuilder) const override;
+	virtual void SetShaderParameters(const FNiagaraDataInterfaceSetShaderParametersContext& Context) const override;
 
-	static const FString BoundsMinName;
-	static const FString BoundsMaxName;
-	static const FString NumPiecesName;
-	static const FString WorldTransformBufferName;
-	static const FString PrevWorldTransformBufferName;
-	static const FString WorldInverseTransformBufferName;
-	static const FString PrevWorldInverseTransformBufferName;
-	static const FString BoundsBufferName;
+	virtual void ProvidePerInstanceDataForRenderThread(void* DataForRenderThread, void* PerInstanceData, const FNiagaraSystemInstanceID& SystemInstance) override;
 
 protected:
 	/** Copy one niagara DI to this */
@@ -189,7 +190,7 @@ struct FNDIGeometryCollectionProxy : public FNiagaraDataInterfaceProxy
 	void DestroyPerInstanceData(const FNiagaraSystemInstanceID& SystemInstance);
 
 	/** Launch all pre stage functions */
-	virtual void PreStage(FRHICommandList& RHICmdList, const FNiagaraDataInterfaceStageArgs& Context) override;
+	virtual void PreStage(const FNDIGpuComputePreStageContext& Context) override;
 
 	/** List of proxy data for each system instances*/
 	TMap<FNiagaraSystemInstanceID, FNDIGeometryCollectionData> SystemInstancesToProxyData;
