- Push MAX_DISTANCE & SMALL_NUMBER into root NiagaraShader and remove from individual DIs

diff --git a/Engine/Plugins/Experimental/ChaosNiagara/Shaders/NiagaraDataInterfaceGeometryCollection.ush b/Engine/Plugins/Experimental/ChaosNiagara/Shaders/NiagaraDataInterfaceGeometryCollection.ush
--- a/Engine/Plugins/Experimental/ChaosNiagara/Shaders/NiagaraDataInterfaceGeometryCollection.ush
+++ b/Engine/Plugins/Experimental/ChaosNiagara/Shaders/NiagaraDataInterfaceGeometryCollection.ush
@@ -9,8 +9,6 @@ NiagaraDataInterfaceGeometryCollection.ush
  * -----------------------------------------------------------------
  */
 
-#define SMALL_NUMBER 1e-8
-#define MAX_DISTANCE 1e+38
 #define BOX_INDEX 0
 #define SPHERE_INDEX 1
 #define CAPSULE_INDEX 2
@@ -210,8 +208,6 @@ void GetClosestPointNoNormal_{ParameterName}(in float3 WorldPosition, in float D
 		OutClosestPosition,OutClosestVelocity);
 }
 
-#undef SMALL_NUMBER
-#undef MAX_DISTANCE
 #undef BOX_INDEX
 #undef SPHERE_INDEX
 #undef CAPSULE_INDEX
diff --git a/Engine/Plugins/FX/Niagara/Shaders/Private/NiagaraDataInterfacePhysicsAsset.ush b/Engine/Plugins/FX/Niagara/Shaders/Private/NiagaraDataInterfacePhysicsAsset.ush
--- a/Engine/Plugins/FX/Niagara/Shaders/Private/NiagaraDataInterfacePhysicsAsset.ush
+++ b/Engine/Plugins/FX/Niagara/Shaders/Private/NiagaraDataInterfacePhysicsAsset.ush
@@ -15,7 +15,6 @@ NiagaraDataInterfacePhysicsAsset.ush
 #define PREVIOUS_OFFSET PHYSICS_ASSET_MAX_TRANSFORMS
 #define REST_OFFSET 2*PHYSICS_ASSET_MAX_TRANSFORMS
 
-#define MAX_DISTANCE 1e+38
 #define BOX_INDEX 0
 #define SPHERE_INDEX 1
 #define CAPSULE_INDEX 2
@@ -513,7 +512,6 @@ void GetProjectionPoint_{ParameterName}(in float3 NodePosition, in float DeltaTi
 #undef PREVIOUS_OFFSET
 #undef REST_OFFSET
 
-#undef MAX_DISTANCE
 #undef BOX_INDEX
 #undef SPHERE_INDEX
 #undef CAPSULE_INDEX
diff --git a/Engine/Plugins/FX/Niagara/Shaders/Private/NiagaraDataInterfaceRigidMeshCollisionQuery.ush b/Engine/Plugins/FX/Niagara/Shaders/Private/NiagaraDataInterfaceRigidMeshCollisionQuery.ush
--- a/Engine/Plugins/FX/Niagara/Shaders/Private/NiagaraDataInterfaceRigidMeshCollisionQuery.ush
+++ b/Engine/Plugins/FX/Niagara/Shaders/Private/NiagaraDataInterfaceRigidMeshCollisionQuery.ush
@@ -9,8 +9,6 @@ NiagaraDataInterfaceRigidMeshCollisionQuery.ush
  * -----------------------------------------------------------------
  */
 
-#define SMALL_NUMBER 1e-8
-#define MAX_DISTANCE 1e+38
 #define BOX_INDEX 0
 #define SPHERE_INDEX 1
 #define CAPSULE_INDEX 2
@@ -499,8 +497,6 @@ void GetClosestPointMeshDistanceFieldNoNormal_{ParameterName}(in float3 WorldPos
 	}	
 }
 
-#undef SMALL_NUMBER
-#undef MAX_DISTANCE
 #undef BOX_INDEX
 #undef SPHERE_INDEX
 #undef CAPSULE_INDEX
diff --git a/Engine/Plugins/FX/Niagara/Shaders/Private/NiagaraEmitterInstanceShader.usf b/Engine/Plugins/FX/Niagara/Shaders/Private/NiagaraEmitterInstanceShader.usf
--- a/Engine/Plugins/FX/Niagara/Shaders/Private/NiagaraEmitterInstanceShader.usf
+++ b/Engine/Plugins/FX/Niagara/Shaders/Private/NiagaraEmitterInstanceShader.usf
@@ -8,6 +8,13 @@ NiagaraSimulationShader.usf:
 
 #include "NiagaraShaderVersion.ush"
 
+#ifndef SMALL_NUMBER
+	#define SMALL_NUMBER 1e-8
+#endif
+#ifndef MAX_DISTANCE
+	#define MAX_DISTANCE 1e+38
+#endif
+
 #if GPU_SIMULATION
 	#include "/Engine/Private/Common.ush"
 	#include "/Engine/Private/DistanceField/GlobalDistanceFieldShared.ush"
diff --git a/Engine/Plugins/FX/Niagara/Shaders/Private/NiagaraQuaternionUtils.ush b/Engine/Plugins/FX/Niagara/Shaders/Private/NiagaraQuaternionUtils.ush
--- a/Engine/Plugins/FX/Niagara/Shaders/Private/NiagaraQuaternionUtils.ush
+++ b/Engine/Plugins/FX/Niagara/Shaders/Private/NiagaraQuaternionUtils.ush
@@ -3,7 +3,9 @@
 #pragma once
 
 #define QUATERNION_IDENTITY float4(0,0,0,1)
-#define SMALL_NUMBER 1e-8
+#ifndef SMALL_NUMBER
+	#define SMALL_NUMBER 1e-8
+#endif
 
 float4 NormalizeQuat(in float4 Quat)
 {
