Removing more Private external directories from include paths.

diff --git a/Engine/Plugins/Experimental/ChaosNiagara/Source/ChaosNiagara/ChaosNiagara.Build.cs b/Engine/Plugins/Experimental/ChaosNiagara/Source/ChaosNiagara/ChaosNiagara.Build.cs
--- a/Engine/Plugins/Experimental/ChaosNiagara/Source/ChaosNiagara/ChaosNiagara.Build.cs
+++ b/Engine/Plugins/Experimental/ChaosNiagara/Source/ChaosNiagara/ChaosNiagara.Build.cs
@@ -8,8 +8,6 @@ public class ChaosNiagara : ModuleRules
 	public ChaosNiagara(ReadOnlyTargetRules Target) : base(Target)
 	{
 		PrivateIncludePaths.AddRange(new string[] {
-			System.IO.Path.Combine(GetModuleDirectory("ControlRig"), "Private"),
-			System.IO.Path.Combine(GetModuleDirectory("Niagara"), "Private"),
 		});
 
 		PrivateDependencyModuleNames.AddRange(
diff --git a/Engine/Plugins/FX/Niagara/Source/Niagara/Niagara.Build.cs b/Engine/Plugins/FX/Niagara/Source/Niagara/Niagara.Build.cs
--- a/Engine/Plugins/FX/Niagara/Source/Niagara/Niagara.Build.cs
+++ b/Engine/Plugins/FX/Niagara/Source/Niagara/Niagara.Build.cs
@@ -67,7 +67,6 @@ public class Niagara : ModuleRules
 
         PrivateIncludePaths.AddRange(
             new string[] {
-				System.IO.Path.Combine(GetModuleDirectory("Engine"), "Private"),
 			});
 
         if (Target.bBuildEditor == true)
diff --git a/Engine/Plugins/FX/Niagara/Source/Niagara/Private/NiagaraGpuComputeDispatch.h b/Engine/Plugins/FX/Niagara/Source/Niagara/Private/NiagaraGpuComputeDispatch.h
--- a/Engine/Plugins/FX/Niagara/Source/Niagara/Private/NiagaraGpuComputeDispatch.h
+++ b/Engine/Plugins/FX/Niagara/Source/Niagara/Private/NiagaraGpuComputeDispatch.h
@@ -21,7 +21,6 @@ the same VectorVM byte code / compute shader code
 #include "NiagaraRendererProperties.h"
 #include "NiagaraSystemGpuComputeProxy.h"
 #include "ParticleResources.h"
-#include "Particles/ParticleSortingGPU.h"
 #include "RenderGraphUtils.h"
 #include "RHIResources.h"
 #include "RendererInterface.h"
