Fix V1078 warnings

diff --git a/Engine/Plugins/Enterprise/DatasmithFBXImporter/Source/DatasmithFBXTranslator/Public/DatasmithFBXScene.h b/Engine/Plugins/Enterprise/DatasmithFBXImporter/Source/DatasmithFBXTranslator/Public/DatasmithFBXScene.h
--- a/Engine/Plugins/Enterprise/DatasmithFBXImporter/Source/DatasmithFBXTranslator/Public/DatasmithFBXScene.h
+++ b/Engine/Plugins/Enterprise/DatasmithFBXImporter/Source/DatasmithFBXTranslator/Public/DatasmithFBXScene.h
@@ -279,12 +279,12 @@ struct DATASMITHFBXTRANSLATOR_API FDatasmithFBXSceneNode : public TSharedFromThi
 
 	void MoveChildren(TSharedPtr<FDatasmithFBXSceneNode>& NewParent)
 	{
-		NewParent->Children = Children;
-		Children.Empty();
 		for (auto& Child : Children)
 		{
 			Child->Parent = NewParent;
 		}
+		NewParent->Children = MoveTemp(Children);
+		Children.Empty();
 	}
 
 	/** Remove this node from hierarchy. */
diff --git a/Engine/Plugins/Experimental/ChaosNiagara/Source/ChaosNiagara/Private/NiagaraDataInterfaceChaosDestruction.cpp b/Engine/Plugins/Experimental/ChaosNiagara/Source/ChaosNiagara/Private/NiagaraDataInterfaceChaosDestruction.cpp
--- a/Engine/Plugins/Experimental/ChaosNiagara/Source/ChaosNiagara/Private/NiagaraDataInterfaceChaosDestruction.cpp
+++ b/Engine/Plugins/Experimental/ChaosNiagara/Source/ChaosNiagara/Private/NiagaraDataInterfaceChaosDestruction.cpp
@@ -1912,16 +1912,11 @@ bool UNiagaraDataInterfaceChaosDestruction::BreakingCallback(FNDIChaosDestructio
 		if (SolverData.Solver->GetEventFilters()->IsBreakingEventEnabled() && BreakingEvents.Num() > 0 && SolverData.Solver->GetSolverTime() > 0.f && MaxNumberOfDataEntriesToSpawn > 0)
 		{
 			TArray<Chaos::FBreakingDataExt>& AllBreakingsArray = BreakingEvents;
-			TMap<IPhysicsProxyBase*, TArray<int32>> AllBreakingsIndicesByPhysicsProxyMap;
 			float TimeData_MapsCreated = SolverData.Solver->GetSolverTime();
 
 			{
 				size_t SizeOfAllBreakings = sizeof(Chaos::FBreakingData) * AllBreakingsArray.Num();
 				size_t SizeOfAllBreakingsIndicesByPhysicsProxy = 0;
-				for (auto& Elem : AllBreakingsIndicesByPhysicsProxyMap)
-				{
-					SizeOfAllBreakingsIndicesByPhysicsProxy += sizeof(int32) * (Elem.Value).Num();
-				}
 				SET_MEMORY_STAT(STAT_AllBreakingsDataMemory, SizeOfAllBreakings);
 				SET_MEMORY_STAT(STAT_AllBreakingsIndicesByPhysicsProxyMemory, SizeOfAllBreakingsIndicesByPhysicsProxy);
 			}
@@ -2384,16 +2379,11 @@ bool UNiagaraDataInterfaceChaosDestruction::TrailingCallback(FNDIChaosDestructio
 		if (SolverData.Solver->GetEventFilters()->IsTrailingEventEnabled() && TrailingEvents.Num() > 0 && SolverData.Solver->GetSolverTime() > 0.f && MaxNumberOfDataEntriesToSpawn > 0)
 		{
 			TArray<Chaos::FTrailingDataExt>& AllTrailingsArray = TrailingEvents;
-			TMap<IPhysicsProxyBase*, TArray<int32>> AllTrailingsIndicesByPhysicsProxyMap;
 			float TimeData_MapsCreated = SolverData.Solver->GetSolverTime();
 
 			{
 				size_t SizeOfAllTrailings = sizeof(Chaos::FTrailingData) * AllTrailingsArray.Num();
 				size_t SizeOfAllTrailingsIndicesByPhysicsProxy = 0;
-				for (auto& Elem : AllTrailingsIndicesByPhysicsProxyMap)
-				{
-					SizeOfAllTrailingsIndicesByPhysicsProxy += sizeof(int32) * (Elem.Value).Num();
-				}
 				SET_MEMORY_STAT(STAT_AllTrailingsDataMemory, SizeOfAllTrailings);
 				SET_MEMORY_STAT(STAT_AllTrailingsIndicesByPhysicsProxyMemory, SizeOfAllTrailingsIndicesByPhysicsProxy);
 			}
diff --git a/Engine/Plugins/FX/Niagara/Source/NiagaraEditor/Private/ViewModels/NiagaraSystemViewModel.cpp b/Engine/Plugins/FX/Niagara/Source/NiagaraEditor/Private/ViewModels/NiagaraSystemViewModel.cpp
--- a/Engine/Plugins/FX/Niagara/Source/NiagaraEditor/Private/ViewModels/NiagaraSystemViewModel.cpp
+++ b/Engine/Plugins/FX/Niagara/Source/NiagaraEditor/Private/ViewModels/NiagaraSystemViewModel.cpp
@@ -663,7 +663,6 @@ void FNiagaraSystemViewModel::DuplicateEmitters(TArray<FEmitterHandleToDuplicate
 	bForceAutoCompileOnce = true;
 
 	// We copy over referenced User Parameters now
-	TArray<TSharedPtr<FNiagaraEmitterHandleViewModel>> NewEmitterHandleViewModels;
 	FNiagaraUserRedirectionParameterStore& TargetParameterStore = GetSystem().GetExposedParameters();
 
 	// we have to keep track of parameters we have already added but renamed.
@@ -777,11 +776,6 @@ void FNiagaraSystemViewModel::DuplicateEmitters(TArray<FEmitterHandleToDuplicate
 		Info.ExpireDuration = 5.f;
 		FSlateNotificationManager::Get().AddNotification(Info);
 	}
-
-	for(TSharedPtr<FNiagaraEmitterHandleViewModel> EmitterHandleViewModel : NewEmitterHandleViewModels)
-	{
-		EmitterHandleViewModel->GetEmitterStackViewModel()->GetRootEntry()->RefreshChildren();
-	}
 }
 
 FGuid FNiagaraSystemViewModel::GetMessageLogGuid() const
diff --git a/Engine/Plugins/Protocols/MQTT/Source/MQTTCore/Private/MQTTConnection.cpp b/Engine/Plugins/Protocols/MQTT/Source/MQTTCore/Private/MQTTConnection.cpp
--- a/Engine/Plugins/Protocols/MQTT/Source/MQTTCore/Private/MQTTConnection.cpp
+++ b/Engine/Plugins/Protocols/MQTT/Source/MQTTCore/Private/MQTTConnection.cpp
@@ -492,12 +492,12 @@ void FMQTTConnection::AbandonOperations()
 
 	{
 		FScopeLock Lock(&IncomingQueueLock);
-		TArray<TUniquePtr<IMQTTOperation>> Messages;
+		//TArray<TUniquePtr<IMQTTOperation>> Messages;
 		//PendingIncomingOperations.GenerateValueArray(Messages);
-		for(const TUniquePtr<IMQTTOperation>& Message : Messages)
-		{
-			Message->Abandon();
-		}
+		//for(const TUniquePtr<IMQTTOperation>& Message : Messages)
+		//{
+		//	Message->Abandon();
+		//}
 		PendingIncomingOperations.Empty();
 	}
 
diff --git a/Engine/Plugins/Runtime/ComputeFramework/Source/ComputeFramework/Private/ComputeKernelFromText.cpp b/Engine/Plugins/Runtime/ComputeFramework/Source/ComputeFramework/Private/ComputeKernelFromText.cpp
--- a/Engine/Plugins/Runtime/ComputeFramework/Source/ComputeFramework/Private/ComputeKernelFromText.cpp
+++ b/Engine/Plugins/Runtime/ComputeFramework/Source/ComputeFramework/Private/ComputeKernelFromText.cpp
@@ -167,6 +167,7 @@ void UComputeKernelFromText::ReparseKernelSourceText()
 
 			FString AllParameters = Matcher.GetCaptureGroup(2);
 			TArray<FString> ParamArray;
+			AllParameters.ParseIntoArray(ParamArray, TEXT(","));
 
 			for(const FString& ParamDecl: ParamArray)
 			{
diff --git a/Engine/Plugins/Runtime/Metasound/Source/MetasoundEditor/Private/MetasoundEditor.cpp b/Engine/Plugins/Runtime/Metasound/Source/MetasoundEditor/Private/MetasoundEditor.cpp
--- a/Engine/Plugins/Runtime/Metasound/Source/MetasoundEditor/Private/MetasoundEditor.cpp
+++ b/Engine/Plugins/Runtime/Metasound/Source/MetasoundEditor/Private/MetasoundEditor.cpp
@@ -1444,16 +1444,17 @@ namespace Metasound
 
 			if (Frontend::ImportJSONAssetToMetasound(InputPath, MetasoundDoc))
 			{
-				TSet<UClass*> ImportClasses;
+				//TSet<UClass*> ImportClasses;
 
 				// TODO: Update importing to support interfaces
 
-				if (ImportClasses.Num() < 1)
+				//if (ImportClasses.Num() < 1)
 				{
 					TArray<FString> InterfaceNames;
 					Algo::Transform(MetasoundDoc.Interfaces, InterfaceNames, [] (const FMetasoundFrontendVersion& InterfaceVersion) { return InterfaceVersion.ToString(); });
 					UE_LOG(LogMetaSound, Warning, TEXT("Cannot create UObject from MetaSound document. No UClass supports interface(s) \"%s\""), *FString::Join(InterfaceNames, TEXT(",")));
 				}
+#if 0
 				else
 				{
 					UClass* AnyClass = nullptr;
@@ -1471,6 +1472,7 @@ namespace Metasound
 
 					// TODO: Update to just use simple UObject NewObject
 				}
+#endif
 			}
 			else
 			{
diff --git a/Engine/Plugins/VirtualProduction/CameraCalibrationCore/Source/CameraCalibrationCore/Private/CameraCalibrationSubsystem.cpp b/Engine/Plugins/VirtualProduction/CameraCalibrationCore/Source/CameraCalibrationCore/Private/CameraCalibrationSubsystem.cpp
--- a/Engine/Plugins/VirtualProduction/CameraCalibrationCore/Source/CameraCalibrationCore/Private/CameraCalibrationSubsystem.cpp
+++ b/Engine/Plugins/VirtualProduction/CameraCalibrationCore/Source/CameraCalibrationCore/Private/CameraCalibrationSubsystem.cpp
@@ -59,16 +59,6 @@ ULensDistortionModelHandlerBase* UCameraCalibrationSubsystem::FindDistortionMode
 ULensDistortionModelHandlerBase* UCameraCalibrationSubsystem::FindOrCreateDistortionModelHandler(FDistortionHandlerPicker& DistortionHandlerPicker, const TSubclassOf<ULensModel> LensModelClass)
 {
 	// This function has been deprecated. The implementation has been changed to provide some backwards compatibility, but code should be updated to not call this function.
-	TArray<ULensDistortionModelHandlerBase*> Handlers = GetDistortionModelHandlers(DistortionHandlerPicker.TargetCameraComponent);
-
-	for (ULensDistortionModelHandlerBase* Handler : Handlers)
-	{
-		if (Handler->GetLensModelClass() == LensModelClass)
-		{
-			return Handler;
-		}
-	}
-
 	return nullptr;
 }
 PRAGMA_ENABLE_DEPRECATION_WARNINGS
