Added getter and setter functions for geometry collection data interface. Allows you to get and modify the transforms for all the elements in a collection.
Also updated existing functionality to work with lwc coordinates.

#tests Tested read and write modules with a custom system; engine test for this DI will be submitted in a separate CL

#virtualized

diff --git a/Engine/Plugins/Experimental/ChaosNiagara/Shaders/NiagaraDataInterfaceGeometryCollection.ush b/Engine/Plugins/Experimental/ChaosNiagara/Shaders/NiagaraDataInterfaceGeometryCollection.ush
--- a/Engine/Plugins/Experimental/ChaosNiagara/Shaders/NiagaraDataInterfaceGeometryCollection.ush
+++ b/Engine/Plugins/Experimental/ChaosNiagara/Shaders/NiagaraDataInterfaceGeometryCollection.ush
@@ -4,6 +4,8 @@
 NiagaraDataInterfaceGeometryCollection.ush 
 =============================================================================*/
 
+#include "/Plugin/FX/Niagara/Private/NiagaraTransformUtils.ush"
+
 /* -----------------------------------------------------------------
  * static mesh constants and context
  * -----------------------------------------------------------------
@@ -15,12 +17,16 @@ NiagaraDataInterfaceGeometryCollection.ush
 
 float3			{ParameterName}_BoundsMin;
 float3			{ParameterName}_BoundsMax;
-uint			{ParameterName}_NumPieces;
+int				{ParameterName}_NumPieces;
+float3			{ParameterName}_RootTransform_Translation;
+float4			{ParameterName}_RootTransform_Rotation;
+float3			{ParameterName}_RootTransform_Scale;
 Buffer<float4>	{ParameterName}_WorldTransformBuffer;
 Buffer<float4>	{ParameterName}_PrevWorldTransformBuffer;
 Buffer<float4>	{ParameterName}_WorldInverseTransformBuffer;
 Buffer<float4>	{ParameterName}_PrevWorldInverseTransformBuffer;
 Buffer<float4>	{ParameterName}_BoundsBuffer;
+ByteAddressBuffer  {ParameterName}_ElementTransforms;
 
 /* -----------------------------------------------------------------
  * Position / Transform utils
@@ -63,7 +69,6 @@ float3x4 DIGeometryCollection_GetPreviousInverse(in int ElementIndex)
     return float3x4({ParameterName}_PrevWorldInverseTransformBuffer[BufferOffset], {ParameterName}_PrevWorldInverseTransformBuffer[BufferOffset+1], {ParameterName}_PrevWorldInverseTransformBuffer[BufferOffset+2]);
 }
 
-
 /* -----------------------------------------------------------------
  * Collision detection utils
  * -----------------------------------------------------------------
@@ -146,6 +151,48 @@ float DIGeometryCollection_GetBoxProjection(in float3 LocalPosition, in float3 B
 	return BoxDistance;
 }
 
+void GetNumElements_{ParameterName}(out int OutCount)
+{
+	OutCount = {ParameterName}_NumPieces;
+}
+
+void GetGeometryComponentTransform_{ParameterName}(out float3 OutPosition, out float4 OutRotation, out float3 OutScale)
+{
+	OutPosition = {ParameterName}_RootTransform_Translation;
+	OutRotation = {ParameterName}_RootTransform_Rotation;
+	OutScale = {ParameterName}_RootTransform_Scale;
+}
+
+void GetElementTransform_{ParameterName}(in int InElementIndex, out float3 OutTranslation, out float4 OutRotation, out float3 OutScale)
+{
+	OutTranslation = 0;
+	OutRotation = 0;
+	OutScale = 0;
+	if (InElementIndex >= 0 && InElementIndex < {ParameterName}_NumPieces)
+	{
+		uint TransformIndex = InElementIndex * 10 * 4;
+		FNiagaraTransform ElementTransform = LoadTransform({ParameterName}_ElementTransforms, TransformIndex);
+		OutTranslation = ElementTransform.Translation;
+		OutRotation = ElementTransform.Rotation;
+		OutScale = ElementTransform.Scale;
+	}
+}
+
+void GetElementBounds_{ParameterName}(in int InElementIndex, out float3 OutCenter, out float3 OutBoundingBox)
+{
+	if (InElementIndex >= 0 && InElementIndex < {ParameterName}_NumPieces)
+	{
+		uint TransformIndex = InElementIndex * 10 * 4;
+		FNiagaraTransform ElementTransform = LoadTransform({ParameterName}_ElementTransforms, TransformIndex);
+		OutCenter = ElementTransform.Translation;
+		OutBoundingBox = {ParameterName}_BoundsBuffer[InElementIndex].xyz;
+	}
+	else
+	{
+		OutCenter = 0;
+		OutBoundingBox = 0;
+	}
+}
 
 /* -----------------------------------------------------------------
  * Get the closest element to the world position
@@ -204,8 +251,7 @@ void DIGeometryCollection_GetClosestPointNoNormal(in float3 WorldPosition, in fl
 
 void GetClosestPointNoNormal_{ParameterName}(in float3 WorldPosition, in float DeltaTime, in float TimeFraction, out float ClosestDistance, out float3 OutClosestPosition, out float3 OutClosestVelocity)
 {
-	DIGeometryCollection_GetClosestPointNoNormal(WorldPosition,DeltaTime,TimeFraction, ClosestDistance,
-		OutClosestPosition,OutClosestVelocity);
+	DIGeometryCollection_GetClosestPointNoNormal(WorldPosition,DeltaTime, TimeFraction, ClosestDistance, OutClosestPosition, OutClosestVelocity);
 }
 
 #undef BOX_INDEX
diff --git a/Engine/Plugins/Experimental/ChaosNiagara/Source/ChaosNiagara/ChaosNiagara.Build.cs b/Engine/Plugins/Experimental/ChaosNiagara/Source/ChaosNiagara/ChaosNiagara.Build.cs
--- a/Engine/Plugins/Experimental/ChaosNiagara/Source/ChaosNiagara/ChaosNiagara.Build.cs
+++ b/Engine/Plugins/Experimental/ChaosNiagara/Source/ChaosNiagara/ChaosNiagara.Build.cs
@@ -20,7 +20,6 @@ public class ChaosNiagara : ModuleRules
 				"RenderCore",
                 "Chaos",
 				"PhysicsCore",
-				"PhysicsCore",
 			}
         );
 					
diff --git a/Engine/Plugins/Experimental/ChaosNiagara/Source/ChaosNiagara/Private/NiagaraDataInterfaceGeometryCollection.cpp b/Engine/Plugins/Experimental/ChaosNiagara/Source/ChaosNiagara/Private/NiagaraDataInterfaceGeometryCollection.cpp
--- a/Engine/Plugins/Experimental/ChaosNiagara/Source/ChaosNiagara/Private/NiagaraDataInterfaceGeometryCollection.cpp
+++ b/Engine/Plugins/Experimental/ChaosNiagara/Source/ChaosNiagara/Private/NiagaraDataInterfaceGeometryCollection.cpp
@@ -1,16 +1,21 @@
 // Copyright Epic Games, Inc. All Rights Reserved.
 
 #include "NiagaraDataInterfaceGeometryCollection.h"
-#include "GeometryCollection/GeometryCollection.h"
-#include "GeometryCollection/GeometryCollectionActor.h"
-#include "GeometryCollection/GeometryCollectionObject.h"
+#include "CoreMinimal.h"
 #include "NiagaraCompileHashVisitor.h"
 #include "NiagaraRenderer.h"
-#include "NiagaraSimStageData.h"
 #include "NiagaraShaderParametersBuilder.h"
+#include "NiagaraSimStageData.h"
 #include "NiagaraSystemInstance.h"
-#include "GeometryCollection/GeometryCollectionComponent.h"
+#include "RenderGraphBuilder.h"
 #include "RenderingThread.h"
+#include "UnifiedBuffer.h"
+#include "Engine/World.h"
+#include "GeometryCollection/GeometryCollection.h"
+#include "GeometryCollection/GeometryCollectionActor.h"
+#include "GeometryCollection/GeometryCollectionAlgo.h"
+#include "GeometryCollection/GeometryCollectionComponent.h"
+#include "GeometryCollection/GeometryCollectionObject.h"
 
 #include UE_INLINE_GENERATED_CPP_BY_NAME(NiagaraDataInterfaceGeometryCollection)
 
@@ -22,6 +27,12 @@ DEFINE_LOG_CATEGORY_STATIC(LogGeometryCollection, Log, All);
 namespace NDIGeometryCollectionLocal
 {
 	static const FName GetClosestPointNoNormalName(TEXT("GetClosestPointNoNormal"));
+	static const FName GetNumElementsName(TEXT("GetNumElements"));
+	static const FName GetElementBoundsName(TEXT("GetElementBounds"));
+	static const FName GetTransformComponentName(TEXT("GetElementTransform"));
+	static const FName SetTransformComponentName(TEXT("SetElementTransform"));
+	static const FName SetTransformWorldName(TEXT("SetElementTransformWS"));
+	static const FName GetComponentWSTransformName(TEXT("GetGeometryComponentTransform"));
 	static const TCHAR* TemplateShaderFilePath = TEXT("/Plugin/Experimental/ChaosNiagara/NiagaraDataInterfaceGeometryCollection.ush");
 
 	template<typename BufferType, EPixelFormat PixelFormat>
@@ -33,7 +44,7 @@ namespace NDIGeometryCollectionLocal
 		}
 	}
 
-	template<typename BufferType, EPixelFormat PixelFormat>
+	template<typename BufferType>
 	void UpdateInternalBuffer(FRHICommandListBase& RHICmdList, const TArray<BufferType>& InputData, FReadBuffer& OutputBuffer)
 	{
 		uint32 ElementCount = InputData.Num();
@@ -47,6 +58,18 @@ namespace NDIGeometryCollectionLocal
 			RHICmdList.UnlockBuffer(OutputBuffer.Buffer);
 		}
 	}
+
+	struct FGeometryCollectionDIFunctionVersion
+	{
+		enum Type
+		{
+			InitialVersion = 0,
+			DIRefactor = 1,
+			
+			VersionPlusOne,
+			LatestVersion = VersionPlusOne - 1
+		};
+	};
 }
 
 //------------------------------------------------------------------------------------------------------------
@@ -93,26 +116,20 @@ void FNDIGeometryCollectionData::Init(UNiagaraDataInterfaceGeometryCollection* I
 {
 	AssetBuffer = nullptr;
 
-	if (Interface != nullptr && SystemInstance != nullptr)
-	{		
-		if (Interface->GeometryCollectionActor != nullptr && Interface->GeometryCollectionActor->GetGeometryCollectionComponent() != nullptr &&
-			Interface->GeometryCollectionActor->GetGeometryCollectionComponent()->RestCollection != nullptr)
+	if (Interface && SystemInstance)
+	{
+		if (const UGeometryCollection* GeometryCollection = Interface->ResolvedSource.GetGeometryCollection())
 		{			
-			const TSharedPtr<FGeometryCollection, ESPMode::ThreadSafe>
-				Collection = Interface->GeometryCollectionActor->GetGeometryCollectionComponent()->RestCollection->GetGeometryCollection();
+			const TSharedPtr<FGeometryCollection> Collection = GeometryCollection->GetGeometryCollection();
 			const TManagedArray<FBox>& BoundingBoxes = Collection->BoundingBox;
-			const TManagedArray<int32>& TransformIndex = Collection->TransformIndex;
-			const TManagedArray<TSet<int32>>& Children = Collection->Children;
 			const TManagedArray<int32>& TransformIndexArray = Collection->TransformIndex;
 
 			int NumPieces = 0;
-			
 			for (int i = 0; i < BoundingBoxes.Num(); ++i)
 			{
 				int32 CurrTransformIndex = TransformIndexArray[i];
 
-				// only consider leaf geometry
-				if (Collection->Children[CurrTransformIndex].Num() == 0)
+				if (Interface->bIncludeIntermediateBones || Collection->Children[CurrTransformIndex].Num() == 0)
 				{
 					NumPieces++;
 				}
@@ -125,20 +142,23 @@ void FNDIGeometryCollectionData::Init(UNiagaraDataInterfaceGeometryCollection* I
 			AssetBuffer->SetNumPieces(NumPieces);
 			BeginInitResource(AssetBuffer);
 
-			FVector Origin;
-			FVector Extents;
-			Interface->GeometryCollectionActor->GetActorBounds(false, Origin, Extents, true);
+			FVector Origin(ForceInitToZero);
+			FVector Extents(ForceInitToZero);
+			if (Interface->ResolvedSource.Component.IsValid())
+			{
+				Interface->ResolvedSource.Component->Bounds.GetBox().GetCenterAndExtents(Origin, Extents);
+			}
 
-			BoundsOrigin = (FVector3f)Origin;	// LWC_TODO: Precision Loss
-			BoundsExtent = (FVector3f)Extents;	// LWC_TODO: Precision Loss
+			FNiagaraLWCConverter LwcConverter = SystemInstance->GetLWCConverter();
+			BoundsOrigin = LwcConverter.ConvertWorldToSimulationVector(Origin);
+			BoundsExtent = static_cast<FVector3f>(Extents);
 
 			int PieceIndex = 0;
 			for (int i = 0; i < BoundingBoxes.Num(); ++i)
 			{
 				int32 CurrTransformIndex = TransformIndexArray[i];
 
-				// only consider leaf geometry
-				if (Collection->Children[CurrTransformIndex].Num() == 0)
+				if (Interface->bIncludeIntermediateBones || Collection->Children[CurrTransformIndex].Num() == 0)
 				{
 					FBox CurrBox = BoundingBoxes[i];
 					FVector3f BoxSize = FVector3f(CurrBox.Max - CurrBox.Min);
@@ -163,30 +183,24 @@ void FNDIGeometryCollectionData::Init(UNiagaraDataInterfaceGeometryCollection* I
 
 void FNDIGeometryCollectionData::Update(UNiagaraDataInterfaceGeometryCollection* Interface, FNiagaraSystemInstance* SystemInstance)
 {
-	if (Interface != nullptr && SystemInstance != nullptr)
+	if (Interface && SystemInstance)
 	{		
 		TickingGroup = ComputeTickingGroup();
 		
-		if (Interface->GeometryCollectionActor != nullptr && Interface->GeometryCollectionActor->GetGeometryCollectionComponent() != nullptr &&
-			Interface->GeometryCollectionActor->GetGeometryCollectionComponent()->RestCollection != nullptr)
-		{			
-			const FTransform ActorTransform = Interface->GeometryCollectionActor->GetTransform();
+		if (const UGeometryCollection* GeoCollection = Interface->ResolvedSource.GetGeometryCollection())
+		{
+			RootTransform = Interface->ResolvedSource.GetComponentRootTransform(SystemInstance);
 
-			const TSharedPtr<FGeometryCollection, ESPMode::ThreadSafe>
-				Collection = Interface->GeometryCollectionActor->GetGeometryCollectionComponent()->RestCollection->GetGeometryCollection();
+			const TSharedPtr<FGeometryCollection> Collection = GeoCollection->GetGeometryCollection();
 			const TManagedArray<FBox>& BoundingBoxes = Collection->BoundingBox;
 			const TManagedArray<int32>& TransformIndexArray = Collection->TransformIndex;
-			const TArray<FTransform3f>& ComponentSpaceTransforms = Interface->GeometryCollectionActor->GetGeometryCollectionComponent()->GetComponentSpaceTransforms3f();
-			const TManagedArray<TSet<int32>>& Children = Collection->Children;
-			
-			int NumPieces = 0;
 
+			int NumPieces = 0;
 			for (int i = 0; i < BoundingBoxes.Num(); ++i)
 			{
 				int32 CurrTransformIndex = TransformIndexArray[i];
 
-				// only consider leaf geometry
-				if (Collection->Children[CurrTransformIndex].Num() == 0)
+				if (Interface->bIncludeIntermediateBones || Collection->Children[CurrTransformIndex].Num() == 0)
 				{
 					NumPieces++;
 				}
@@ -196,20 +210,28 @@ void FNDIGeometryCollectionData::Update(UNiagaraDataInterfaceGeometryCollection*
 			{
 				Init(Interface, SystemInstance);
 			}
+			AssetArrays->ComponentRestTransformBuffer = Interface->ResolvedSource.GetInitialLocalRestTransforms();
 			
-			FVector Origin;
-			FVector Extents;
-			Interface->GeometryCollectionActor->GetActorBounds(false, Origin, Extents, true);
+			FVector Origin(ForceInitToZero);
+			FVector Extents(ForceInitToZero);
+			if (Interface->ResolvedSource.Component.IsValid())
+			{
+				Interface->ResolvedSource.Component->Bounds.GetBox().GetCenterAndExtents(Origin, Extents);
+			}
 
-			BoundsOrigin = (FVector3f)Origin;
-			BoundsExtent = (FVector3f)Extents;
+			FNiagaraLWCConverter LwcConverter = SystemInstance->GetLWCConverter();
+			BoundsOrigin = LwcConverter.ConvertWorldToSimulationVector(Origin);
+			BoundsExtent = static_cast<FVector3f>(Extents);
 
 			int PieceIndex = 0;
 			for (int i = 0; i < BoundingBoxes.Num(); ++i)
 			{
 				int32 CurrTransformIndex = TransformIndexArray[i];
-				if (Collection->Children[CurrTransformIndex].Num() == 0)
+				if (Interface->bIncludeIntermediateBones || Collection->Children[CurrTransformIndex].Num() == 0)
 				{
+					ensureMsgf(AssetArrays->ComponentRestTransformBuffer.IsValidIndex(CurrTransformIndex), TEXT("The local rest transforms and the indices for the transform mapping should always match."));
+					AssetArrays->ElementIndexToTransformBufferMapping[PieceIndex] = CurrTransformIndex;
+					
 					int32 TransformIndex = 3 * PieceIndex;
 					AssetArrays->PrevWorldInverseTransformBuffer[TransformIndex] = AssetArrays->WorldInverseTransformBuffer[TransformIndex];
 					AssetArrays->PrevWorldInverseTransformBuffer[TransformIndex + 1] = AssetArrays->WorldInverseTransformBuffer[TransformIndex + 1];
@@ -225,7 +247,7 @@ void FNDIGeometryCollectionData::Update(UNiagaraDataInterfaceGeometryCollection*
 					FVector LocalTranslation = (CurrBox.Max + CurrBox.Min) * .5;
 					FTransform LocalOffset(LocalTranslation);
 
-					const FTransform3f CurrTransform(LocalOffset * FTransform(ComponentSpaceTransforms[CurrTransformIndex]) * ActorTransform);
+					const FTransform3f CurrTransform(LocalOffset * Interface->ResolvedSource.GetComponentSpaceTransform(CurrTransformIndex) * RootTransform);
 					CurrTransform.ToMatrixWithScale().To3x4MatrixTranspose(&AssetArrays->WorldTransformBuffer[TransformIndex].X);
 
 					const FTransform3f CurrInverse = CurrTransform.Inverse();
@@ -244,37 +266,84 @@ ETickingGroup FNDIGeometryCollectionData::ComputeTickingGroup()
 	return TickingGroup;
 }
 
-//------------------------------------------------------------------------------------------------------------
+const UGeometryCollection* FResolvedNiagaraGeometryCollection::GetGeometryCollection() const
+{
+	if (Collection.IsValid())
+	{
+		return Collection.Get();
+	}
+	if (Component.IsValid())
+	{
+		return Component->GetRestCollection();
+	}
+	return nullptr;
+}
 
-void FNDIGeometryCollectionProxy::ConsumePerInstanceDataFromGameThread(void* DataFromGameThread, const FNiagaraSystemInstanceID& Instance)
+FTransform FResolvedNiagaraGeometryCollection::GetComponentRootTransform(FNiagaraSystemInstance* SystemInstance) const
 {
-	check(IsInRenderingThread());
+	FTransform ComponentTransform = FTransform::Identity;
+	if (Component.IsValid())
+	{
+		ComponentTransform = Component->GetComponentTransform();
+	}
+	ComponentTransform.AddToTranslation(FVector(SystemInstance->GetLWCTile()) * -FLargeWorldRenderScalar::GetTileSize());
+	return ComponentTransform;
+}
 
-	FNDIGeometryCollectionData* SourceData = static_cast<FNDIGeometryCollectionData*>(DataFromGameThread);
-	FNDIGeometryCollectionData* TargetData = &(SystemInstancesToProxyData.FindOrAdd(Instance));
+FTransform FResolvedNiagaraGeometryCollection::GetComponentSpaceTransform(int32 TransformIndex) const
+{
+	if (Component.IsValid())
+	{
+		const TArray<FTransform3f>& ComponentSpaceTransforms = Component->GetComponentSpaceTransforms3f();
+		if (ComponentSpaceTransforms.IsValidIndex(TransformIndex))
+		{
+			return FTransform(ComponentSpaceTransforms[TransformIndex]);
+		}
+	}
+	return FTransform();
+}
 
-	ensure(TargetData);
-	if (TargetData)
+TArray<FTransform> FResolvedNiagaraGeometryCollection::GetInitialLocalRestTransforms() const
+{
+	if (Component.IsValid())
 	{
-		TargetData->AssetBuffer = SourceData->AssetBuffer;		
-		TargetData->AssetArrays = SourceData->AssetArrays;
-		TargetData->TickingGroup = SourceData->TickingGroup;
-		TargetData->BoundsOrigin = SourceData->BoundsOrigin;
-		TargetData->BoundsExtent = SourceData->BoundsExtent;
+		return Component->GetInitialLocalRestTransforms();
 	}
-	else
+	if (Collection.IsValid())
 	{
-		UE_LOG(LogGeometryCollection, Log, TEXT("ConsumePerInstanceDataFromGameThread() ... could not find %d"), Instance);
+		TArray<FTransform> InitialLocalTransforms;
+		const FGeometryCollection& RestGeometryCollection = *Collection->GetGeometryCollection();
+		GeometryCollectionAlgo::GlobalMatrices(RestGeometryCollection.Transform, RestGeometryCollection.Parent, InitialLocalTransforms);
+		return InitialLocalTransforms;
 	}
+	return TArray<FTransform>();
+}
+
+//------------------------------------------------------------------------------------------------------------
+
+void FNDIGeometryCollectionProxy::ConsumePerInstanceDataFromGameThread(void* DataFromGameThread, const FNiagaraSystemInstanceID& Instance)
+{
+	check(IsInRenderingThread());
+
+	FNDIGeometryCollectionData* SourceData = static_cast<FNDIGeometryCollectionData*>(DataFromGameThread);
+	FNDIGeometryCollectionData& TargetData = SystemInstancesToProxyData.FindOrAdd(Instance);
+
+	TargetData.AssetBuffer = SourceData->AssetBuffer;		
+	TargetData.AssetArrays = SourceData->AssetArrays;
+	TargetData.TickingGroup = SourceData->TickingGroup;
+	TargetData.RootTransform = SourceData->RootTransform;
+	TargetData.BoundsOrigin = SourceData->BoundsOrigin;
+	TargetData.BoundsExtent = SourceData->BoundsExtent;
+	
 	SourceData->~FNDIGeometryCollectionData();
 }
 
 void FNDIGeometryCollectionProxy::InitializePerInstanceData(const FNiagaraSystemInstanceID& SystemInstance)
 {
 	check(IsInRenderingThread());
-
 	check(!SystemInstancesToProxyData.Contains(SystemInstance));
-	FNDIGeometryCollectionData* TargetData = &SystemInstancesToProxyData.Add(SystemInstance);
+	
+	SystemInstancesToProxyData.Add(SystemInstance);
 }
 
 void FNDIGeometryCollectionProxy::DestroyPerInstanceData(const FNiagaraSystemInstanceID& SystemInstance)
@@ -293,15 +362,36 @@ void FNDIGeometryCollectionProxy::PreStage(const FNDIGpuComputePreStageContext&
 	{
 		if (Context.GetSimStageData().bFirstStage)
 		{
-			FNDIGeometryCollectionBuffer* AssetBuffer = ProxyData->AssetBuffer;
 			FRHICommandListBase& RHICmdList = FRHICommandListImmediate::Get();
 
 			// #todo(dmp): bounds buffer doesn't need to be updated each frame
-			NDIGeometryCollectionLocal::UpdateInternalBuffer<FVector4f, EPixelFormat::PF_A32B32G32R32F>(RHICmdList, ProxyData->AssetArrays->WorldTransformBuffer, ProxyData->AssetBuffer->WorldTransformBuffer);
-			NDIGeometryCollectionLocal::UpdateInternalBuffer<FVector4f, EPixelFormat::PF_A32B32G32R32F>(RHICmdList, ProxyData->AssetArrays->PrevWorldTransformBuffer, ProxyData->AssetBuffer->PrevWorldTransformBuffer);
-			NDIGeometryCollectionLocal::UpdateInternalBuffer<FVector4f, EPixelFormat::PF_A32B32G32R32F>(RHICmdList, ProxyData->AssetArrays->WorldInverseTransformBuffer, ProxyData->AssetBuffer->WorldInverseTransformBuffer);
-			NDIGeometryCollectionLocal::UpdateInternalBuffer<FVector4f, EPixelFormat::PF_A32B32G32R32F>(RHICmdList, ProxyData->AssetArrays->PrevWorldInverseTransformBuffer, ProxyData->AssetBuffer->PrevWorldInverseTransformBuffer);
-			NDIGeometryCollectionLocal::UpdateInternalBuffer<FVector4f, EPixelFormat::PF_A32B32G32R32F>(RHICmdList, ProxyData->AssetArrays->BoundsBuffer, ProxyData->AssetBuffer->BoundsBuffer);
+			NDIGeometryCollectionLocal::UpdateInternalBuffer<FVector4f>(RHICmdList, ProxyData->AssetArrays->WorldTransformBuffer, ProxyData->AssetBuffer->WorldTransformBuffer);
+			NDIGeometryCollectionLocal::UpdateInternalBuffer<FVector4f>(RHICmdList, ProxyData->AssetArrays->PrevWorldTransformBuffer, ProxyData->AssetBuffer->PrevWorldTransformBuffer);
+			NDIGeometryCollectionLocal::UpdateInternalBuffer<FVector4f>(RHICmdList, ProxyData->AssetArrays->WorldInverseTransformBuffer, ProxyData->AssetBuffer->WorldInverseTransformBuffer);
+			NDIGeometryCollectionLocal::UpdateInternalBuffer<FVector4f>(RHICmdList, ProxyData->AssetArrays->PrevWorldInverseTransformBuffer, ProxyData->AssetBuffer->PrevWorldInverseTransformBuffer);
+			NDIGeometryCollectionLocal::UpdateInternalBuffer<FVector4f>(RHICmdList, ProxyData->AssetArrays->BoundsBuffer, ProxyData->AssetBuffer->BoundsBuffer);
+
+			// build RDG buffer with per-element component transforms
+			TArray<uint8>& DataToUpload = ProxyData->AssetBuffer->DataToUpload;
+			if (DataToUpload.Num() > 0)
+			{
+				FRDGBuilder& GraphBuilder = Context.GetGraphBuilder();
+
+				const uint64 InitialDataSize = DataToUpload.Num() * DataToUpload.GetTypeSize();
+				const uint64 BufferSize = Align(InitialDataSize, 16);
+
+				FRDGBufferDesc BufferDesc = FRDGBufferDesc::CreateByteAddressDesc(BufferSize);
+				ResizeBufferIfNeeded(GraphBuilder, ProxyData->AssetBuffer->ComponentRestTransformBuffer, BufferDesc, TEXT("NiagaraGeometryCollection"));
+
+				GraphBuilder.QueueBufferUpload(
+					GraphBuilder.RegisterExternalBuffer(ProxyData->AssetBuffer->ComponentRestTransformBuffer),
+					DataToUpload.GetData(),
+					InitialDataSize,
+					ERDGInitialDataFlags::None
+				);
+
+				DataToUpload.Empty();
+			}
 		}
 	}
 }
@@ -319,6 +409,9 @@ bool UNiagaraDataInterfaceGeometryCollection::InitPerInstanceData(void* PerInsta
 	FNDIGeometryCollectionData* InstanceData = new (PerInstanceData) FNDIGeometryCollectionData();
 
 	check(InstanceData);
+	FNiagaraParameterDirectBinding<UObject*> CollectionParameterBinding;
+	CollectionParameterBinding.Init(SystemInstance->GetInstanceParameters(), GeometryCollectionUserParameter.Parameter);
+	ResolveGeometryCollection(SystemInstance, CollectionParameterBinding.GetValue());
 	InstanceData->Init(this, SystemInstance);
 	
 	return true;
@@ -326,9 +419,7 @@ bool UNiagaraDataInterfaceGeometryCollection::InitPerInstanceData(void* PerInsta
 
 ETickingGroup UNiagaraDataInterfaceGeometryCollection::CalculateTickGroup(const void* PerInstanceData) const
 {
-	const FNDIGeometryCollectionData* InstanceData = static_cast<const FNDIGeometryCollectionData*>(PerInstanceData);
-
-	if (InstanceData)
+	if (const FNDIGeometryCollectionData* InstanceData = static_cast<const FNDIGeometryCollectionData*>(PerInstanceData))
 	{
 		return InstanceData->TickingGroup;
 	}
@@ -377,11 +468,167 @@ bool UNiagaraDataInterfaceGeometryCollection::CopyToInternal(UNiagaraDataInterfa
 
 	UNiagaraDataInterfaceGeometryCollection* OtherTyped = CastChecked<UNiagaraDataInterfaceGeometryCollection>(Destination);		
 
+	OtherTyped->SourceMode = SourceMode;
+#if WITH_EDITORONLY_DATA
+	OtherTyped->PreviewCollection = PreviewCollection;
+#endif
+	OtherTyped->DefaultGeometryCollection = DefaultGeometryCollection;
 	OtherTyped->GeometryCollectionActor = GeometryCollectionActor;
+	OtherTyped->SourceComponent = SourceComponent;
+	OtherTyped->GeometryCollectionUserParameter = GeometryCollectionUserParameter;
+	OtherTyped->bIncludeIntermediateBones = bIncludeIntermediateBones;
+	OtherTyped->ResolvedSource = ResolvedSource;
+
+	return true;
+}
+
+void UNiagaraDataInterfaceGeometryCollection::ResolveGeometryCollection(FNiagaraSystemInstance* SystemInstance, UObject* UserParameter)
+{
+	ResolvedSource = FResolvedNiagaraGeometryCollection();
 
+	switch (SourceMode)
+	{
+	case ENDIGeometryCollection_SourceMode::Source:
+		ResolveGeometryCollectionFromDirectSource();
+		break;
+	case ENDIGeometryCollection_SourceMode::AttachParent:
+		ResolveGeometryCollectionFromAttachParent(SystemInstance);
+		break;
+	case ENDIGeometryCollection_SourceMode::DefaultCollectionOnly:
+		ResolveGeometryCollectionFromDefaultCollection();
+		break;
+	case ENDIGeometryCollection_SourceMode::ParameterBinding:
+		ResolveGeometryCollectionFromParamterBinding(UserParameter);
+		break;
+	case ENDIGeometryCollection_SourceMode::Default:
+	default:
+		if (!ResolveGeometryCollectionFromDirectSource())
+		{
+			if (!ResolveGeometryCollectionFromParamterBinding(UserParameter))
+			{
+				if (!ResolveGeometryCollectionFromAttachParent(SystemInstance))
+				{
+					ResolveGeometryCollectionFromDefaultCollection();
+				}
+			}
+		}
+		break;
+	}
+
+#if WITH_EDITORONLY_DATA
+	if (!ResolvedSource.Collection.IsValid() && !ResolvedSource.Component.IsValid() && (!SystemInstance || !SystemInstance->GetWorld()->IsGameWorld()))
+	{
+		// NOTE: We don't fall back on the preview mesh if we have a valid collection referenced
+		ResolvedSource.Collection = PreviewCollection.LoadSynchronous();
+	}
+#endif
+}
+
+bool UNiagaraDataInterfaceGeometryCollection::ResolveGeometryCollectionFromDirectSource()
+{
+	if (::IsValid(SourceComponent))
+	{
+		ResolvedSource.Component = SourceComponent;
+		return true;
+	}
+	if (GeometryCollectionActor.Get())
+	{
+		ResolvedSource.Component = GeometryCollectionActor->GetGeometryCollectionComponent();
+	}
+	return false;
+}
+
+bool UNiagaraDataInterfaceGeometryCollection::ResolveGeometryCollectionFromAttachParent(FNiagaraSystemInstance* SystemInstance)
+{
+	if (USceneComponent* AttachComponent = SystemInstance->GetAttachComponent())
+	{
+		// First, try to find the geometry collection component up the attachment hierarchy
+		for (USceneComponent* Curr = AttachComponent; Curr; Curr = Curr->GetAttachParent())
+		{
+			UGeometryCollectionComponent* ParentComp = Cast<UGeometryCollectionComponent>(Curr);
+			if (::IsValid(ParentComp))
+			{
+				ResolvedSource.Component = ParentComp;
+				return true;
+			}
+		}
+
+		// Next, try to find one in our outer chain
+		UGeometryCollectionComponent* OuterComp = AttachComponent->GetTypedOuter<UGeometryCollectionComponent>();
+		if (::IsValid(OuterComp))
+		{
+			ResolvedSource.Component = OuterComp;
+			return true;
+		}
+		return ResolveGeometryCollectionFromActor(AttachComponent->GetAttachmentRootActor());
+	}
+	return false;
+}
+
+bool UNiagaraDataInterfaceGeometryCollection::ResolveGeometryCollectionFromActor(AActor* Actor)
+{
+	if (Actor)
+	{
+		// Final fall-back, look for any component on our root actor or any of its parents
+		if (AGeometryCollectionActor* GeoActor = Cast<AGeometryCollectionActor>(Actor))
+		{
+			ResolvedSource.Component = GeoActor->GetGeometryCollectionComponent();
+			return true;
+		}
+
+		// Fall back on any valid component on the actor
+		for (UActorComponent* ActorComp : Actor->GetComponents())
+		{
+			UGeometryCollectionComponent* Comp = Cast<UGeometryCollectionComponent>(ActorComp);
+			if (::IsValid(Comp))
+			{
+				ResolvedSource.Component = Comp;
+				return true;
+			}
+		}
+	}
+	return false;
+}
+
+bool UNiagaraDataInterfaceGeometryCollection::ResolveGeometryCollectionFromDefaultCollection()
+{
+	ResolvedSource.Collection = DefaultGeometryCollection;
 	return true;
 }
 
+bool UNiagaraDataInterfaceGeometryCollection::ResolveGeometryCollectionFromParamterBinding(UObject* ParameterBindingValue)
+{
+	if (ParameterBindingValue == nullptr)
+	{
+		return false;
+	}
+	if (ResolveGeometryCollectionFromActor(Cast<AActor>(ParameterBindingValue)))
+	{
+		return true;
+	}
+	if (UGeometryCollectionComponent* UserComponent = Cast<UGeometryCollectionComponent>(ParameterBindingValue))
+	{
+		ResolvedSource.Component = UserComponent;
+		return true;
+	}
+	if (UGeometryCollection* UserCollection = Cast<UGeometryCollection>(ParameterBindingValue))
+	{
+		ResolvedSource.Collection = UserCollection;
+		return true;
+	}
+	return false;
+}
+
+bool UNiagaraDataInterfaceGeometryCollection::PerInstanceTickPostSimulate(void* PerInstanceData, FNiagaraSystemInstance* InSystemInstance, float DeltaSeconds)
+{
+	FNDIGeometryCollectionData* InstanceData = static_cast<FNDIGeometryCollectionData*>(PerInstanceData);
+	if (InstanceData && InstanceData->bHasPendingComponentTransformUpdate && ResolvedSource.Component.IsValid())
+	{
+		ResolvedSource.Component->SetLocalRestTransforms(InstanceData->AssetArrays->ComponentRestTransformBuffer, !bIncludeIntermediateBones);
+	}
+	return false;
+}
+
 bool UNiagaraDataInterfaceGeometryCollection::Equals(const UNiagaraDataInterface* Other) const
 {
 	if (!Super::Equals(Other))
@@ -390,7 +637,15 @@ bool UNiagaraDataInterfaceGeometryCollection::Equals(const UNiagaraDataInterface
 	}
 	const UNiagaraDataInterfaceGeometryCollection* OtherTyped = CastChecked<const UNiagaraDataInterfaceGeometryCollection>(Other);
 
-	return  OtherTyped->GeometryCollectionActor == GeometryCollectionActor;
+	return  OtherTyped->SourceMode == SourceMode &&
+#if WITH_EDITORONLY_DATA
+			OtherTyped->PreviewCollection == PreviewCollection &&
+#endif
+			OtherTyped->DefaultGeometryCollection == DefaultGeometryCollection &&
+			OtherTyped->GeometryCollectionActor == GeometryCollectionActor &&
+			OtherTyped->SourceComponent == SourceComponent &&
+			OtherTyped->GeometryCollectionUserParameter == GeometryCollectionUserParameter &&
+			OtherTyped->bIncludeIntermediateBones == bIncludeIntermediateBones;
 }
 
 void UNiagaraDataInterfaceGeometryCollection::PostInitProperties()
@@ -405,22 +660,119 @@ void UNiagaraDataInterfaceGeometryCollection::PostInitProperties()
 }
 
 #if WITH_EDITORONLY_DATA
+
 void UNiagaraDataInterfaceGeometryCollection::GetFunctionsInternal(TArray<FNiagaraFunctionSignature>& OutFunctions) const
 {
 	{
 		FNiagaraFunctionSignature Sig;
 		Sig.Name = NDIGeometryCollectionLocal::GetClosestPointNoNormalName;
-		Sig.bSupportsGPU = true;
 		Sig.bSupportsCPU = false;
 		Sig.bMemberFunction = true;
-		Sig.bRequiresContext = false;
-		Sig.Inputs.Add(FNiagaraVariable(FNiagaraTypeDefinition(GetClass()), TEXT("Collision DI")));
-		Sig.Inputs.Add(FNiagaraVariable(FNiagaraTypeDefinition::GetVec3Def(), TEXT("World Position")));
-		Sig.Inputs.Add(FNiagaraVariable(FNiagaraTypeDefinition::GetFloatDef(), TEXT("Delta Time")));
-		Sig.Inputs.Add(FNiagaraVariable(FNiagaraTypeDefinition::GetFloatDef(), TEXT("Time Fraction")));
+		Sig.FunctionVersion = NDIGeometryCollectionLocal::FGeometryCollectionDIFunctionVersion::LatestVersion;
+		Sig.Description = LOCTEXT("SigNoNormalNameDescription", "Returns the closest point on the surface of the geometry collection.");
+		Sig.Inputs.Add(FNiagaraVariable(FNiagaraTypeDefinition(GetClass()), TEXT("Geometry Collection DI")));
+		Sig.Inputs.Add(FNiagaraVariable(FNiagaraTypeDefinition::GetPositionDef(), TEXT("World Position")));
+		Sig.AddInput(FNiagaraVariable(FNiagaraTypeDefinition::GetFloatDef(), TEXT("Delta Time")), LOCTEXT("ClosestPointNoNormal_DeltaTimeTooltip", "Current delta time to compute the returned velocity"));
+		Sig.AddInput(FNiagaraVariable(FNiagaraTypeDefinition::GetFloatDef(), TEXT("Time Fraction")), LOCTEXT("ClosestPointNoNormal_TimeFractionTooltip", "Lerps the returned closest position between the current frame (1.0) and the previous frame (0.0)."));
 		Sig.Outputs.Add(FNiagaraVariable(FNiagaraTypeDefinition::GetFloatDef(), TEXT("Closest Distance")));
-		Sig.Outputs.Add(FNiagaraVariable(FNiagaraTypeDefinition::GetVec3Def(), TEXT("Closest Position")));		
+		Sig.Outputs.Add(FNiagaraVariable(FNiagaraTypeDefinition::GetPositionDef(), TEXT("Closest Position")));
 		Sig.Outputs.Add(FNiagaraVariable(FNiagaraTypeDefinition::GetVec3Def(), TEXT("Closest Velocity")));
+		NIAGARA_ADD_FUNCTION_SOURCE_INFO(Sig)
+
+		OutFunctions.Add(Sig);
+	}
+
+	{
+		FNiagaraFunctionSignature Sig;
+		Sig.Name = NDIGeometryCollectionLocal::GetNumElementsName;
+		Sig.bMemberFunction = true;
+		Sig.FunctionVersion = NDIGeometryCollectionLocal::FGeometryCollectionDIFunctionVersion::LatestVersion;
+		Sig.Description = LOCTEXT("SigNumElementsDescription", "Returns the numbers of elements in the geometry collection. Unless 'IncludeIntermediateBones' is set, this only counts leaf geometries.");
+		Sig.Inputs.Add(FNiagaraVariable(FNiagaraTypeDefinition(GetClass()), TEXT("Geometry Collection DI")));
+		Sig.Outputs.Add(FNiagaraVariable(FNiagaraTypeDefinition::GetIntDef(), TEXT("Count")));
+		NIAGARA_ADD_FUNCTION_SOURCE_INFO(Sig)
+
+		OutFunctions.Add(Sig);
+	}
+
+	{
+		FNiagaraFunctionSignature Sig;
+		Sig.Name = NDIGeometryCollectionLocal::GetElementBoundsName;
+		Sig.bMemberFunction = true;
+		Sig.FunctionVersion = NDIGeometryCollectionLocal::FGeometryCollectionDIFunctionVersion::LatestVersion;
+		Sig.Description = LOCTEXT("GetElementBoundsNameDescription", "Returns the current bounding box and size for the given element. The values are relative to the geometry component root.");
+		Sig.Inputs.Add(FNiagaraVariable(FNiagaraTypeDefinition(GetClass()), TEXT("Geometry Collection DI")));
+		Sig.Inputs.Add(FNiagaraVariable(FNiagaraTypeDefinition::GetIntDef(), TEXT("Element Index")));
+		Sig.Outputs.Add(FNiagaraVariable(FNiagaraTypeDefinition::GetVec3Def(), TEXT("Bounding Box Center")));
+		Sig.Outputs.Add(FNiagaraVariable(FNiagaraTypeDefinition::GetVec3Def(), TEXT("Bounding Box Size")));
+		NIAGARA_ADD_FUNCTION_SOURCE_INFO(Sig)
+
+		OutFunctions.Add(Sig);
+	}
+
+	{
+		FNiagaraFunctionSignature Sig;
+		Sig.Name = NDIGeometryCollectionLocal::GetTransformComponentName;
+		Sig.bMemberFunction = true;
+		Sig.FunctionVersion = NDIGeometryCollectionLocal::FGeometryCollectionDIFunctionVersion::LatestVersion;
+		Sig.Description = LOCTEXT("SigGetTransformComponentNameDescription", "Returns the transform for the given element index relative to the root of the geometry collection.");
+		Sig.Inputs.Add(FNiagaraVariable(FNiagaraTypeDefinition(GetClass()), TEXT("Geometry Collection DI")));
+		Sig.Inputs.Add(FNiagaraVariable(FNiagaraTypeDefinition::GetIntDef(), TEXT("Element Index")));
+		Sig.Outputs.Add(FNiagaraVariable(FNiagaraTypeDefinition::GetVec3Def(), TEXT("Translation")));
+		Sig.Outputs.Add(FNiagaraVariable(FNiagaraTypeDefinition::GetQuatDef(), TEXT("Rotation")));
+		Sig.Outputs.Add(FNiagaraVariable(FNiagaraTypeDefinition::GetVec3Def(), TEXT("Scale")));
+		NIAGARA_ADD_FUNCTION_SOURCE_INFO(Sig)
+
+		OutFunctions.Add(Sig);
+	}
+
+	{
+		FNiagaraFunctionSignature Sig;
+		Sig.Name = NDIGeometryCollectionLocal::SetTransformComponentName;
+		Sig.bSupportsGPU = false; //TODO: add gpu support
+		Sig.bMemberFunction = true;
+		Sig.bRequiresExecPin = true;
+		Sig.FunctionVersion = NDIGeometryCollectionLocal::FGeometryCollectionDIFunctionVersion::LatestVersion;
+		Sig.Description = LOCTEXT("SigSetTransformComponentNameDescription", "Sets the transform for the given element index relative to the root of the geometry collection.");
+		Sig.Inputs.Add(FNiagaraVariable(FNiagaraTypeDefinition(GetClass()), TEXT("Geometry Collection DI")));
+		Sig.Inputs.Add(FNiagaraVariable(FNiagaraTypeDefinition::GetIntDef(), TEXT("Element Index")));
+		Sig.Inputs.Add(FNiagaraVariable(FNiagaraTypeDefinition::GetVec3Def(), TEXT("Translation")));
+		Sig.Inputs.Add(FNiagaraVariable(FNiagaraTypeDefinition::GetQuatDef(), TEXT("Rotation")));
+		Sig.Inputs.Add(FNiagaraVariable(FNiagaraTypeDefinition::GetVec3Def(), TEXT("Scale")));
+		NIAGARA_ADD_FUNCTION_SOURCE_INFO(Sig)
+
+		OutFunctions.Add(Sig);
+	}
+
+	{
+		FNiagaraFunctionSignature Sig;
+		Sig.Name = NDIGeometryCollectionLocal::SetTransformWorldName;
+		Sig.bSupportsGPU = false; //TODO: add gpu support
+		Sig.bMemberFunction = true;
+		Sig.bRequiresExecPin = true;
+		Sig.FunctionVersion = NDIGeometryCollectionLocal::FGeometryCollectionDIFunctionVersion::LatestVersion;
+		Sig.Description = LOCTEXT("SigSetTransformWorldNameDescription", "Sets the transform for the given element index in world space.");
+		Sig.Inputs.Add(FNiagaraVariable(FNiagaraTypeDefinition(GetClass()), TEXT("Geometry Collection DI")));
+		Sig.Inputs.Add(FNiagaraVariable(FNiagaraTypeDefinition::GetIntDef(), TEXT("Element Index")));
+		Sig.Inputs.Add(FNiagaraVariable(FNiagaraTypeDefinition::GetPositionDef(), TEXT("Position")));
+		Sig.Inputs.Add(FNiagaraVariable(FNiagaraTypeDefinition::GetQuatDef(), TEXT("Rotation")));
+		Sig.Inputs.Add(FNiagaraVariable(FNiagaraTypeDefinition::GetVec3Def(), TEXT("Scale")));
+		NIAGARA_ADD_FUNCTION_SOURCE_INFO(Sig)
+
+		OutFunctions.Add(Sig);
+	}
+
+	{
+		FNiagaraFunctionSignature Sig;
+		Sig.Name = NDIGeometryCollectionLocal::GetComponentWSTransformName;
+		Sig.bMemberFunction = true;
+		Sig.FunctionVersion = NDIGeometryCollectionLocal::FGeometryCollectionDIFunctionVersion::LatestVersion;
+		Sig.Description = LOCTEXT("SigGetActorTransformNameDescription", "Returns the transform for the geometry collection component (or its owning actor) in world space.");
+		Sig.Inputs.Add(FNiagaraVariable(FNiagaraTypeDefinition(GetClass()), TEXT("Geometry Collection DI")));
+		Sig.Outputs.Add(FNiagaraVariable(FNiagaraTypeDefinition::GetPositionDef(), TEXT("Position")));
+		Sig.Outputs.Add(FNiagaraVariable(FNiagaraTypeDefinition::GetQuatDef(), TEXT("Rotation")));
+		Sig.Outputs.Add(FNiagaraVariable(FNiagaraTypeDefinition::GetVec3Def(), TEXT("Scale")));
+		NIAGARA_ADD_FUNCTION_SOURCE_INFO(Sig)
 
 		OutFunctions.Add(Sig);
 	}
@@ -429,12 +781,203 @@ void UNiagaraDataInterfaceGeometryCollection::GetFunctionsInternal(TArray<FNiaga
 
 void UNiagaraDataInterfaceGeometryCollection::GetVMExternalFunction(const FVMExternalFunctionBindingInfo& BindingInfo, void* InstanceData, FVMExternalFunction& OutFunc)
 {
+	using namespace NDIGeometryCollectionLocal;
+
+	if ( BindingInfo.Name == GetNumElementsName )
+	{
+		OutFunc = FVMExternalFunction::CreateUObject(this, &UNiagaraDataInterfaceGeometryCollection::GetNumGeometryElements);
+	}
+	else if ( BindingInfo.Name == GetElementBoundsName )
+	{
+		OutFunc = FVMExternalFunction::CreateUObject(this, &UNiagaraDataInterfaceGeometryCollection::GetElementBounds);
+	}
+	else if ( BindingInfo.Name == GetTransformComponentName )
+	{
+		OutFunc = FVMExternalFunction::CreateUObject(this, &UNiagaraDataInterfaceGeometryCollection::GetElementTransformCS);
+	}
+	else if ( BindingInfo.Name == SetTransformComponentName )
+	{
+		OutFunc = FVMExternalFunction::CreateUObject(this, &UNiagaraDataInterfaceGeometryCollection::SetElementTransformCS);
+	}
+	else if ( BindingInfo.Name == SetTransformWorldName )
+	{
+		OutFunc = FVMExternalFunction::CreateUObject(this, &UNiagaraDataInterfaceGeometryCollection::SetElementTransformWS);
+	}
+	else if ( BindingInfo.Name == GetComponentWSTransformName )
+	{
+		OutFunc = FVMExternalFunction::CreateUObject(this, &UNiagaraDataInterfaceGeometryCollection::GetActorTransform);
+	}
+}
+
+void UNiagaraDataInterfaceGeometryCollection::GetNumGeometryElements(FVectorVMExternalFunctionContext& Context)
+{
+	VectorVM::FUserPtrHandler<FNDIGeometryCollectionData> InstanceData(Context);
+	FNDIOutputParam<int32> OutCount(Context);
+
+	for (int32 i = 0; i < Context.GetNumInstances(); ++i)
+	{
+		OutCount.SetAndAdvance(InstanceData.Get()->AssetArrays->NumPieces);
+	}
+}
+
+void UNiagaraDataInterfaceGeometryCollection::GetElementBounds(FVectorVMExternalFunctionContext& Context)
+{
+	VectorVM::FUserPtrHandler<FNDIGeometryCollectionData> InstanceData(Context);
+	FNDIInputParam<int32> InElement(Context);
+	FNDIOutputParam<FVector3f> OutCenter(Context);
+	FNDIOutputParam<FVector3f> OutBoundingBoxSize(Context);
+
+	for (int32 i = 0; i < Context.GetNumInstances(); ++i)
+	{
+		int32 ElementIndex = InElement.GetAndAdvance();
+		int32 TransformIndex = INDEX_NONE;
+		if (InstanceData.Get()->AssetArrays->ElementIndexToTransformBufferMapping.IsValidIndex(ElementIndex))
+		{
+			TransformIndex = InstanceData.Get()->AssetArrays->ElementIndexToTransformBufferMapping[ElementIndex];
+		}
+		if (InstanceData.Get()->AssetArrays->ComponentRestTransformBuffer.IsValidIndex(TransformIndex))
+		{
+			const FTransform& Transform = InstanceData.Get()->AssetArrays->ComponentRestTransformBuffer[TransformIndex];
+			FVector3f Center(Transform.GetTranslation());
+			FVector4f BoundsExtent = InstanceData.Get()->AssetArrays->BoundsBuffer[ElementIndex];
+
+			OutCenter.SetAndAdvance(Center);
+			OutBoundingBoxSize.SetAndAdvance(FVector3f(BoundsExtent));
+		}
+		else
+		{
+			OutCenter.SetAndAdvance(FVector3f::ZeroVector);
+			OutBoundingBoxSize.SetAndAdvance(FVector3f::ZeroVector);
+		}
+	}
+}
+
+void UNiagaraDataInterfaceGeometryCollection::GetElementTransformCS(FVectorVMExternalFunctionContext& Context)
+{
+	VectorVM::FUserPtrHandler<FNDIGeometryCollectionData> InstanceData(Context);
+	FNDIInputParam<int32> InElement(Context);
+	FNDIOutputParam<FVector3f> OutTranslation(Context);
+	FNDIOutputParam<FQuat4f> OutRotation(Context);
+	FNDIOutputParam<FVector3f> OutScale(Context);
+
+	for (int32 i = 0; i < Context.GetNumInstances(); ++i)
+	{
+		int32 ElementIndex = InElement.GetAndAdvance();
+		int32 TransformIndex = INDEX_NONE;
+		if (InstanceData.Get()->AssetArrays->ElementIndexToTransformBufferMapping.IsValidIndex(ElementIndex))
+		{
+			TransformIndex = InstanceData.Get()->AssetArrays->ElementIndexToTransformBufferMapping[ElementIndex];
+		}
+		if (InstanceData.Get()->AssetArrays->ComponentRestTransformBuffer.IsValidIndex(TransformIndex))
+		{
+			const FTransform& Transform = InstanceData.Get()->AssetArrays->ComponentRestTransformBuffer[TransformIndex];
+			OutTranslation.SetAndAdvance(FVector3f(Transform.GetTranslation()));
+			OutRotation.SetAndAdvance(FQuat4f(Transform.GetRotation().GetNormalized()));
+			OutScale.SetAndAdvance(FVector3f(Transform.GetScale3D()));
+		}
+		else
+		{
+			OutTranslation.SetAndAdvance(FVector3f::ZeroVector);
+			OutRotation.SetAndAdvance(FQuat4f::Identity);
+			OutScale.SetAndAdvance(FVector3f::OneVector);
+		}
+	}
+}
+
+void UNiagaraDataInterfaceGeometryCollection::SetElementTransformWS(FVectorVMExternalFunctionContext& Context)
+{
+	VectorVM::FUserPtrHandler<FNDIGeometryCollectionData> InstanceData(Context);
+	FNDIInputParam<int32> InElement(Context);
+	FNDIInputParam<FNiagaraPosition> InTranslation(Context);
+	FNDIInputParam<FQuat4f> InRotation(Context);
+	FNDIInputParam<FVector3f> InScale(Context);
+
+	const FTransform& InverseRootTransform = InstanceData.Get()->RootTransform.Inverse();
+	for (int32 i = 0; i < Context.GetNumInstances(); ++i)
+	{
+		int32 ElementIndex = InElement.GetAndAdvance();
+		FVector Translation(InTranslation.GetAndAdvance());
+		FQuat Rotation(InRotation.GetAndAdvance());
+		FVector Scale(InScale.GetAndAdvance());
+
+		if (InstanceData.Get()->AssetArrays->ElementIndexToTransformBufferMapping.IsValidIndex(ElementIndex))
+		{
+			int32 TransformIndex = InstanceData.Get()->AssetArrays->ElementIndexToTransformBufferMapping[ElementIndex];
+			if (InstanceData.Get()->AssetArrays->ComponentRestTransformBuffer.IsValidIndex(TransformIndex))
+			{
+				FTransform Transform(Rotation, Translation, Scale);
+				// the geometry component wants all the transforms to be in localspace, so we need to remove the component root transform (which also removes the LWC tile offset)
+				Transform *= InverseRootTransform;
+				if (!Transform.Equals(InstanceData.Get()->AssetArrays->ComponentRestTransformBuffer[TransformIndex]))
+				{
+					InstanceData.Get()->AssetArrays->ComponentRestTransformBuffer[TransformIndex] = Transform;
+					InstanceData.Get()->bHasPendingComponentTransformUpdate = true;
+				}
+			}
+		}
+	}
+}
+
+void UNiagaraDataInterfaceGeometryCollection::SetElementTransformCS(FVectorVMExternalFunctionContext& Context)
+{
+	VectorVM::FUserPtrHandler<FNDIGeometryCollectionData> InstanceData(Context);
+	FNDIInputParam<int32> InElement(Context);
+	FNDIInputParam<FVector3f> InTranslation(Context);
+	FNDIInputParam<FQuat4f> InRotation(Context);
+	FNDIInputParam<FVector3f> InScale(Context);
+
+	for (int32 i = 0; i < Context.GetNumInstances(); ++i)
+	{
+		int32 ElementIndex = InElement.GetAndAdvance();
+		FVector Translation(InTranslation.GetAndAdvance());
+		FQuat Rotation(InRotation.GetAndAdvance());
+		FVector Scale(InScale.GetAndAdvance());
+
+		if (InstanceData.Get()->AssetArrays->ElementIndexToTransformBufferMapping.IsValidIndex(ElementIndex))
+		{
+			int32 TransformIndex = InstanceData.Get()->AssetArrays->ElementIndexToTransformBufferMapping[ElementIndex];
+			if (InstanceData.Get()->AssetArrays->ComponentRestTransformBuffer.IsValidIndex(TransformIndex))
+			{
+				FTransform NewTransform(Rotation, Translation, Scale);
+				FTransform& CurrentTransform = InstanceData.Get()->AssetArrays->ComponentRestTransformBuffer[TransformIndex];
+				if (!CurrentTransform.Equals(NewTransform))
+				{
+					InstanceData.Get()->AssetArrays->ComponentRestTransformBuffer[TransformIndex] = NewTransform;
+					InstanceData.Get()->bHasPendingComponentTransformUpdate = true;
+				}
+			}
+		}
+	}
+}
+
+void UNiagaraDataInterfaceGeometryCollection::GetActorTransform(FVectorVMExternalFunctionContext& Context)
+{
+	VectorVM::FUserPtrHandler<FNDIGeometryCollectionData> InstanceData(Context);
+	FNDIOutputParam<FNiagaraPosition> OutPosition(Context);
+	FNDIOutputParam<FQuat4f> OutRotation(Context);
+	FNDIOutputParam<FVector3f> OutScale(Context);
+
+	const FTransform& Transform = InstanceData.Get()->RootTransform;
+	FNiagaraPosition Location(Transform.GetTranslation());
+	FQuat4f Rotation(Transform.GetRotation().GetNormalized());
+	FVector3f Scale(Transform.GetScale3D());
+	
+	for (int32 i = 0; i < Context.GetNumInstances(); ++i)
+	{
+		OutPosition.SetAndAdvance(Location);
+		OutRotation.SetAndAdvance(Rotation);
+		OutScale.SetAndAdvance(Scale);
+	}
 }
 
 #if WITH_EDITORONLY_DATA
 bool UNiagaraDataInterfaceGeometryCollection::GetFunctionHLSL(const FNiagaraDataInterfaceGPUParamInfo& ParamInfo, const FNiagaraDataInterfaceGeneratedFunction& FunctionInfo, int FunctionInstanceIndex, FString& OutHLSL)
 {
-	if (FunctionInfo.DefinitionName == NDIGeometryCollectionLocal::GetClosestPointNoNormalName)
+	if (FunctionInfo.DefinitionName == NDIGeometryCollectionLocal::GetClosestPointNoNormalName ||
+		FunctionInfo.DefinitionName == NDIGeometryCollectionLocal::GetNumElementsName ||
+		FunctionInfo.DefinitionName == NDIGeometryCollectionLocal::GetComponentWSTransformName ||
+		FunctionInfo.DefinitionName == NDIGeometryCollectionLocal::GetElementBoundsName ||
+		FunctionInfo.DefinitionName == NDIGeometryCollectionLocal::GetTransformComponentName)
 	{
 		return true;
 	}
@@ -442,6 +985,28 @@ bool UNiagaraDataInterfaceGeometryCollection::GetFunctionHLSL(const FNiagaraData
 	return false;
 }
 
+
+
+bool UNiagaraDataInterfaceGeometryCollection::UpgradeFunctionCall(FNiagaraFunctionSignature& FunctionSignature)
+{
+	// upgrade old functions to the latest version
+	if (FunctionSignature.FunctionVersion < NDIGeometryCollectionLocal::FGeometryCollectionDIFunctionVersion::DIRefactor)
+	{
+		TArray<FNiagaraFunctionSignature> AllFunctions;
+		GetFunctionsInternal(AllFunctions);
+		for (const FNiagaraFunctionSignature& Sig : AllFunctions)
+		{
+			if (FunctionSignature.Name == Sig.Name && Sig.Name == NDIGeometryCollectionLocal::GetClosestPointNoNormalName)
+			{
+				FunctionSignature = Sig;
+				return true;
+			}
+		}
+	}
+
+	return false;
+}
+
 bool UNiagaraDataInterfaceGeometryCollection::AppendCompileHash(FNiagaraCompileHashVisitor* InVisitor) const
 {
 	bool bSuccess = Super::AppendCompileHash(InVisitor);
@@ -465,31 +1030,25 @@ void UNiagaraDataInterfaceGeometryCollection::BuildShaderParameters(FNiagaraShad
 void UNiagaraDataInterfaceGeometryCollection::SetShaderParameters(const FNiagaraDataInterfaceSetShaderParametersContext& Context) const
 {
 	FNDIGeometryCollectionProxy& InterfaceProxy = Context.GetProxy<FNDIGeometryCollectionProxy>();
-	FNDIGeometryCollectionData* ProxyData = InterfaceProxy.SystemInstancesToProxyData.Find(Context.GetSystemInstanceID());
+	FNDIGeometryCollectionData& ProxyData = InterfaceProxy.SystemInstancesToProxyData.FindChecked(Context.GetSystemInstanceID());
+	
+	FRDGBuilder& GraphBuilder = Context.GetGraphBuilder();
+	FRDGBufferRef RDGBuffer = GraphBuilder.RegisterExternalBuffer(ProxyData.AssetBuffer->ComponentRestTransformBuffer);
+	FRDGBufferSRVRef RDGBufferSRV = GraphBuilder.CreateSRV(RDGBuffer);
 
 	FShaderParameters* ShaderParameters = Context.GetParameterNestedStruct<FShaderParameters>();
-	if (ProxyData != nullptr)
-	{
-		ShaderParameters->BoundsMin							= ProxyData->BoundsOrigin - ProxyData->BoundsExtent;
-		ShaderParameters->BoundsMax							= ProxyData->BoundsOrigin + ProxyData->BoundsExtent;
-		ShaderParameters->NumPieces							= ProxyData->AssetBuffer->NumPieces;
-		ShaderParameters->WorldTransformBuffer				= FNiagaraRenderer::GetSrvOrDefaultFloat4(ProxyData->AssetBuffer->WorldTransformBuffer.SRV);
-		ShaderParameters->PrevWorldTransformBuffer			= FNiagaraRenderer::GetSrvOrDefaultFloat4(ProxyData->AssetBuffer->PrevWorldTransformBuffer.SRV);
-		ShaderParameters->WorldInverseTransformBuffer		= FNiagaraRenderer::GetSrvOrDefaultFloat4(ProxyData->AssetBuffer->WorldInverseTransformBuffer.SRV);
-		ShaderParameters->PrevWorldInverseTransformBuffer	= FNiagaraRenderer::GetSrvOrDefaultFloat4(ProxyData->AssetBuffer->PrevWorldInverseTransformBuffer.SRV);
-		ShaderParameters->BoundsBuffer						= FNiagaraRenderer::GetSrvOrDefaultFloat4(ProxyData->AssetBuffer->BoundsBuffer.SRV);
-	}
-	else
-	{
-		ShaderParameters->BoundsMin							= FVector3f::ZeroVector;
-		ShaderParameters->BoundsMax							= FVector3f::ZeroVector;
-		ShaderParameters->NumPieces							= 0;
-		ShaderParameters->WorldTransformBuffer				= FNiagaraRenderer::GetDummyFloat4Buffer();
-		ShaderParameters->PrevWorldTransformBuffer			= FNiagaraRenderer::GetDummyFloat4Buffer();
-		ShaderParameters->WorldInverseTransformBuffer		= FNiagaraRenderer::GetDummyFloat4Buffer();
-		ShaderParameters->PrevWorldInverseTransformBuffer	= FNiagaraRenderer::GetDummyFloat4Buffer();
-		ShaderParameters->BoundsBuffer						= FNiagaraRenderer::GetDummyFloat4Buffer();
-	}
+	ShaderParameters->BoundsMin							= ProxyData.BoundsOrigin - ProxyData.BoundsExtent;
+	ShaderParameters->BoundsMax							= ProxyData.BoundsOrigin + ProxyData.BoundsExtent;
+	ShaderParameters->NumPieces							= ProxyData.AssetBuffer->NumPieces;
+	ShaderParameters->RootTransform_Translation			= FVector3f(ProxyData.RootTransform.GetTranslation());
+	ShaderParameters->RootTransform_Rotation			= FQuat4f(ProxyData.RootTransform.GetRotation());
+	ShaderParameters->RootTransform_Scale				= FVector3f(ProxyData.RootTransform.GetScale3D());
+	ShaderParameters->WorldTransformBuffer				= FNiagaraRenderer::GetSrvOrDefaultFloat4(ProxyData.AssetBuffer->WorldTransformBuffer.SRV);
+	ShaderParameters->PrevWorldTransformBuffer			= FNiagaraRenderer::GetSrvOrDefaultFloat4(ProxyData.AssetBuffer->PrevWorldTransformBuffer.SRV);
+	ShaderParameters->WorldInverseTransformBuffer		= FNiagaraRenderer::GetSrvOrDefaultFloat4(ProxyData.AssetBuffer->WorldInverseTransformBuffer.SRV);
+	ShaderParameters->PrevWorldInverseTransformBuffer	= FNiagaraRenderer::GetSrvOrDefaultFloat4(ProxyData.AssetBuffer->PrevWorldInverseTransformBuffer.SRV);
+	ShaderParameters->BoundsBuffer						= FNiagaraRenderer::GetSrvOrDefaultFloat4(ProxyData.AssetBuffer->BoundsBuffer.SRV);
+	ShaderParameters->ElementTransforms					= RDGBufferSRV;
 }
 
 void UNiagaraDataInterfaceGeometryCollection::ProvidePerInstanceDataForRenderThread(void* DataForRenderThread, void* PerInstanceData, const FNiagaraSystemInstanceID& SystemInstance)
@@ -506,6 +1065,25 @@ void UNiagaraDataInterfaceGeometryCollection::ProvidePerInstanceDataForRenderThr
 		RenderThreadData->TickingGroup = GameThreadData->TickingGroup;
 		RenderThreadData->BoundsOrigin = GameThreadData->BoundsOrigin;
 		RenderThreadData->BoundsExtent = GameThreadData->BoundsExtent;
+		RenderThreadData->RootTransform = GameThreadData->RootTransform;
+
+		constexpr int32 TransformGpuSize = 10 * 4;
+		const int32 BufferSize = RenderThreadData->AssetArrays->ElementIndexToTransformBufferMapping.Num() * TransformGpuSize;
+		RenderThreadData->AssetBuffer->DataToUpload.SetNumUninitialized(BufferSize);
+		float* OutFloats = reinterpret_cast<float*>(RenderThreadData->AssetBuffer->DataToUpload.GetData());
+		for (int32 Index : RenderThreadData->AssetArrays->ElementIndexToTransformBufferMapping)
+		{
+			const FTransform& Transform = RenderThreadData->AssetArrays->ComponentRestTransformBuffer[Index];
+			
+			const FVector3f Translation(Transform.GetTranslation());
+			const FQuat4f Rotation(Transform.GetRotation());
+			const FVector3f Scale(Transform.GetScale3D());
+
+			FMemory::Memcpy(&OutFloats[0], &Translation, sizeof(FVector3f));
+			FMemory::Memcpy(&OutFloats[3], &Rotation, sizeof(FQuat4f));
+			FMemory::Memcpy(&OutFloats[7], &Scale, sizeof(FVector3f));
+			OutFloats += 10;
+		}
 	}
 	check(Proxy);
 }
diff --git a/Engine/Plugins/Experimental/ChaosNiagara/Source/ChaosNiagara/Public/NiagaraDataInterfaceGeometryCollection.h b/Engine/Plugins/Experimental/ChaosNiagara/Source/ChaosNiagara/Public/NiagaraDataInterfaceGeometryCollection.h
--- a/Engine/Plugins/Experimental/ChaosNiagara/Source/ChaosNiagara/Public/NiagaraDataInterfaceGeometryCollection.h
+++ b/Engine/Plugins/Experimental/ChaosNiagara/Source/ChaosNiagara/Public/NiagaraDataInterfaceGeometryCollection.h
@@ -4,6 +4,7 @@
 
 #include "NiagaraDataInterface.h"
 #include "RHIUtilities.h"
+#include "GeometryCollection/GeometryCollectionObject.h"
 
 #include "NiagaraDataInterfaceGeometryCollection.generated.h"
 
@@ -19,6 +20,8 @@ struct FNDIGeometryCollectionArrays
 	TArray<FVector4f> WorldInverseTransformBuffer;
 	TArray<FVector4f> PrevWorldInverseTransformBuffer;
 	TArray<FVector4f> BoundsBuffer;
+	TArray<FTransform> ComponentRestTransformBuffer;
+	TArray<uint32> ElementIndexToTransformBufferMapping;
 
 	FNDIGeometryCollectionArrays()
 	{
@@ -39,6 +42,8 @@ struct FNDIGeometryCollectionArrays
 		WorldInverseTransformBuffer = Other->WorldInverseTransformBuffer;
 		PrevWorldInverseTransformBuffer = Other->PrevWorldInverseTransformBuffer;
 		BoundsBuffer = Other->BoundsBuffer;
+		ComponentRestTransformBuffer = Other->ComponentRestTransformBuffer;
+		ElementIndexToTransformBufferMapping = Other->ElementIndexToTransformBufferMapping;
 	}
 
 	void Resize(uint32 Num)
@@ -50,6 +55,8 @@ struct FNDIGeometryCollectionArrays
 		WorldInverseTransformBuffer.Init(FVector4f(0, 0, 0, 0), 3 * NumPieces);
 		PrevWorldInverseTransformBuffer.Init(FVector4f(0, 0, 0, 0), 3 * NumPieces);
 		BoundsBuffer.Init(FVector4f(0, 0, 0, 0), NumPieces);
+		ComponentRestTransformBuffer.Init(FTransform(), 1);
+		ElementIndexToTransformBufferMapping.Init(0, NumPieces);
 	}
 
 	uint32 NumPieces = 100;
@@ -82,10 +89,16 @@ struct FNDIGeometryCollectionBuffer : public FRenderResource
 	/** Element extent buffer */
 	FReadBuffer BoundsBuffer;
 
+	/** Per-element transform buffer */
+	TRefCountPtr<FRDGPooledBuffer>	ComponentRestTransformBuffer;
+
+	/** Raw data for the transform buffer */
+	TArray<uint8> DataToUpload;
+
 	/** number of transforms */
-	uint32 NumPieces;
+	int32 NumPieces;
 
-	void SetNumPieces(uint32 Num)
+	void SetNumPieces(int32 Num)
 	{
 		NumPieces = Num;
 	}
@@ -108,6 +121,9 @@ struct FNDIGeometryCollectionData
 	/** The instance ticking group */
 	ETickingGroup TickingGroup;
 
+	/** Actor or geometry collection component world transform, adjusted by lwc system tile */
+	FTransform RootTransform;
+	
 	/** Geometry Collection Bounds */
 	FVector3f BoundsOrigin;
 
@@ -117,11 +133,55 @@ struct FNDIGeometryCollectionData
 	FNDIGeometryCollectionBuffer* AssetBuffer = nullptr;
 
 	/** Physics asset Cpu arrays */
-	FNDIGeometryCollectionArrays *AssetArrays = nullptr;	
+	FNDIGeometryCollectionArrays* AssetArrays = nullptr;
+
+	bool bHasPendingComponentTransformUpdate = false;
+};
+
+UENUM()
+enum class ENDIGeometryCollection_SourceMode : uint8
+{
+	/**
+	Default behavior follows the order of.
+	- Use "Source" when specified (either set explicitly or via blueprint with Set Niagara Geometry Collection Component).
+	- Use Parameter Binding if valid
+	- Find Geometry Collection Component, Attached Actor, Attached Component
+	- Falls back to 'Default Geometry Collection' specified on the data interface
+	*/
+	Default,
+
+	/**	Only use "Source" (either set explicitly or via blueprint with Set Niagara Geometry Collection Component). */
+	Source,
+
+	/**	Only use the parent actor or component the system is attached to. */
+	AttachParent,
+
+	/** Only use the "Default Geometry Collection" specified. */
+	DefaultCollectionOnly,
+
+	/** Only use the parameter binding. */
+	ParameterBinding,
+};
+
+USTRUCT()
+struct FResolvedNiagaraGeometryCollection
+{
+	GENERATED_BODY()
+	
+	const UGeometryCollection* GetGeometryCollection() const;
+	FTransform GetComponentRootTransform(FNiagaraSystemInstance* SystemInstance) const;
+	FTransform GetComponentSpaceTransform(int32 TransformIndex) const;
+	TArray<FTransform> GetInitialLocalRestTransforms() const;
+	
+	UPROPERTY()
+	TWeakObjectPtr<UGeometryCollection> Collection;
+
+	UPROPERTY()
+	TWeakObjectPtr<UGeometryCollectionComponent> Component;
 };
 
 /** Data Interface for the Collisions */
-UCLASS(EditInlineNew, Category = "Collision", meta = (DisplayName = "Geometry Collection"))
+UCLASS(EditInlineNew, Category = "Chaos", meta = (DisplayName = "Geometry Collection"), MinimalAPI)
 class UNiagaraDataInterfaceGeometryCollection : public UNiagaraDataInterface
 {
 	GENERATED_UCLASS_BODY()
@@ -129,30 +189,66 @@ class UNiagaraDataInterfaceGeometryCollection : public UNiagaraDataInterface
 	BEGIN_SHADER_PARAMETER_STRUCT(FShaderParameters, )
 		SHADER_PARAMETER(FVector3f,				BoundsMin)
 		SHADER_PARAMETER(FVector3f,				BoundsMax)
-		SHADER_PARAMETER(uint32,				NumPieces)
+		SHADER_PARAMETER(int32,					NumPieces)
+		SHADER_PARAMETER(FVector3f,				RootTransform_Translation)
+		SHADER_PARAMETER(FQuat4f,				RootTransform_Rotation)
+		SHADER_PARAMETER(FVector3f,				RootTransform_Scale)
 		SHADER_PARAMETER_SRV(Buffer<float4>,	WorldTransformBuffer)
 		SHADER_PARAMETER_SRV(Buffer<float4>,	PrevWorldTransformBuffer)
 		SHADER_PARAMETER_SRV(Buffer<float4>,	WorldInverseTransformBuffer)
 		SHADER_PARAMETER_SRV(Buffer<float4>,	PrevWorldInverseTransformBuffer)
 		SHADER_PARAMETER_SRV(Buffer<float4>,	BoundsBuffer)
+		SHADER_PARAMETER_RDG_BUFFER_SRV(ByteAddressBuffer,	ElementTransforms)
 	END_SHADER_PARAMETER_STRUCT()
 
-public:
+	/** Controls how to retrieve the Skeletal Mesh Component to attach to. */
 	UPROPERTY(EditAnywhere, Category = "Geometry Collection")
-	TObjectPtr<AGeometryCollectionActor> GeometryCollectionActor;
+	ENDIGeometryCollection_SourceMode SourceMode = ENDIGeometryCollection_SourceMode::Default;
+
+#if WITH_EDITORONLY_DATA
+	/** Geometry collection used to sample from when not overridden by a source actor from the scene. Only available in editor for previewing. This is removed in cooked builds. */
+	UPROPERTY(EditAnywhere, Category = "Geometry Collection")
+	TSoftObjectPtr<UGeometryCollection> PreviewCollection;
+#endif
+
+	/** GeometryCollection used to sample from when not overridden by a source actor from the scene. This reference is NOT removed from cooked builds. */
+	UPROPERTY(EditAnywhere, Category = "Geometry Collection", meta = (EditConditionHides, EditCondition = "SourceMode == ENDIGeometryCollection_SourceMode::Default || SourceMode == ENDISocketReaderSourceMode::DefaultCollectionOnly"))
+	TObjectPtr<UGeometryCollection> DefaultGeometryCollection;
+
+	/** The source actor from which to sample. Takes precedence over the direct geometry collection. Note that this can only be set when used as a user variable on a niagara component in the world.*/
+	UPROPERTY(EditAnywhere, Category = "Geometry Collection", meta = (DisplayName = "Source Actor", EditConditionHides, EditCondition = "SourceMode == ENDIGeometryCollection_SourceMode::Default || SourceMode == ENDIGeometryCollection_SourceMode::Source"))
+	TSoftObjectPtr<AGeometryCollectionActor> GeometryCollectionActor;
+
+	/** The source component from which to sample. Takes precedence over the direct mesh. Not exposed to the user, only indirectly accessible from blueprints. */
+	UPROPERTY(Transient)
+	TObjectPtr<UGeometryCollectionComponent> SourceComponent;
+
+	/** Reference to a user parameter if we're reading one. */
+	UPROPERTY(EditAnywhere, Category = "Geometry Collection", meta = (EditConditionHides, EditCondition = "SourceMode == ENDIGeometryCollection_SourceMode::Default || SourceMode == ENDISocketReaderSourceMode::ParameterBinding"))
+	FNiagaraUserParameterBinding GeometryCollectionUserParameter;
+	
+	// If true then this data interface will also read and write intermediate bones or geometry, otherwise only leaf nodes are considered
+	UPROPERTY(EditAnywhere, Category = "Geometry Collection")
+	bool bIncludeIntermediateBones = false;
+
+	UPROPERTY(Transient)
+	FResolvedNiagaraGeometryCollection ResolvedSource;
 
 	/** UObject Interface */
 	virtual void PostInitProperties() override;
 
 	/** UNiagaraDataInterface Interface */
 	virtual void GetVMExternalFunction(const FVMExternalFunctionBindingInfo& BindingInfo, void* InstanceData, FVMExternalFunction& OutFunc) override;
-	virtual bool CanExecuteOnTarget(ENiagaraSimTarget Target) const override { return Target == ENiagaraSimTarget::GPUComputeSim; }
+	virtual bool CanExecuteOnTarget(ENiagaraSimTarget Target) const override { return true; }
 	virtual bool InitPerInstanceData(void* PerInstanceData, FNiagaraSystemInstance* SystemInstance) override;
 	virtual void DestroyPerInstanceData(void* PerInstanceData, FNiagaraSystemInstance* SystemInstance)override;
 	virtual bool PerInstanceTick(void* PerInstanceData, FNiagaraSystemInstance* SystemInstance, float DeltaSeconds) override;
 	virtual int32 PerInstanceDataSize()const override { return sizeof(FNDIGeometryCollectionData); }
+	virtual bool PerInstanceTickPostSimulate(void* PerInstanceData, FNiagaraSystemInstance* InSystemInstance, float DeltaSeconds) override;
 	virtual bool Equals(const UNiagaraDataInterface* Other) const override;
 	virtual bool HasPreSimulateTick() const override { return true; }
+	virtual bool HasPostSimulateTick() const override { return true; }
+	virtual bool PostSimulateCanOverlapFrames() const override { return false; }
 	virtual bool HasTickGroupPrereqs() const override { return true; }
 	virtual ETickingGroup CalculateTickGroup(const void* PerInstanceData) const override;
 
@@ -162,12 +258,22 @@ public:
 	virtual bool AppendCompileHash(FNiagaraCompileHashVisitor* InVisitor) const override;
 	virtual void GetParameterDefinitionHLSL(const FNiagaraDataInterfaceGPUParamInfo& ParamInfo, FString& OutHLSL) override;
 	virtual bool GetFunctionHLSL(const FNiagaraDataInterfaceGPUParamInfo& ParamInfo, const FNiagaraDataInterfaceGeneratedFunction& FunctionInfo, int FunctionInstanceIndex, FString& OutHLSL) override;
+	virtual bool UpgradeFunctionCall(FNiagaraFunctionSignature& FunctionSignature) override;
 #endif
 	virtual void BuildShaderParameters(FNiagaraShaderParametersBuilder& ShaderParametersBuilder) const override;
 	virtual void SetShaderParameters(const FNiagaraDataInterfaceSetShaderParametersContext& Context) const override;
 
 	virtual void ProvidePerInstanceDataForRenderThread(void* DataForRenderThread, void* PerInstanceData, const FNiagaraSystemInstanceID& SystemInstance) override;
 
+private:
+	// VM functions
+	void GetNumGeometryElements(FVectorVMExternalFunctionContext& Context);
+	void GetElementBounds(FVectorVMExternalFunctionContext& Context);
+	void GetElementTransformCS(FVectorVMExternalFunctionContext& Context);
+	void SetElementTransformCS(FVectorVMExternalFunctionContext& Context);
+	void SetElementTransformWS(FVectorVMExternalFunctionContext& Context);
+	void GetActorTransform(FVectorVMExternalFunctionContext& Context);
+
 protected:
 #if WITH_EDITORONLY_DATA
 	virtual void GetFunctionsInternal(TArray<FNiagaraFunctionSignature>& OutFunctions) const override;
@@ -175,6 +281,14 @@ protected:
 
 	/** Copy one niagara DI to this */
 	virtual bool CopyToInternal(UNiagaraDataInterface* Destination) const override;
+
+private:
+	void ResolveGeometryCollection(FNiagaraSystemInstance* SystemInstance, UObject* UserParameter);
+	bool ResolveGeometryCollectionFromDirectSource();
+	bool ResolveGeometryCollectionFromAttachParent(FNiagaraSystemInstance* SystemInstance);
+	bool ResolveGeometryCollectionFromActor(AActor* Actor);
+	bool ResolveGeometryCollectionFromDefaultCollection();
+	bool ResolveGeometryCollectionFromParamterBinding(UObject* UserParameter);
 };
 
 /** Proxy to send data to gpu */
@@ -198,7 +312,3 @@ struct FNDIGeometryCollectionProxy : public FNiagaraDataInterfaceProxy
 	/** List of proxy data for each system instances*/
 	TMap<FNiagaraSystemInstanceID, FNDIGeometryCollectionData> SystemInstancesToProxyData;
 };
-
-#if UE_ENABLE_INCLUDE_ORDER_DEPRECATED_IN_5_2
-#include "GeometryCollection/GeometryCollectionActor.h"
-#endif
diff --git a/Engine/Plugins/FX/Niagara/Shaders/Private/NiagaraTransformUtils.ush b/Engine/Plugins/FX/Niagara/Shaders/Private/NiagaraTransformUtils.ush
--- a/Engine/Plugins/FX/Niagara/Shaders/Private/NiagaraTransformUtils.ush
+++ b/Engine/Plugins/FX/Niagara/Shaders/Private/NiagaraTransformUtils.ush
@@ -16,7 +16,7 @@ FNiagaraTransform LoadTransform(ByteAddressBuffer Buffer, uint Offset)
 	FNiagaraTransform R;
 	R.Translation	= asfloat(Buffer.Load3(Offset + 0));
 	R.Rotation		= asfloat(Buffer.Load4(Offset + 12));
-	R.Scale			= asfloat(Buffer.Load3(Offset + 20));
+	R.Scale			= asfloat(Buffer.Load3(Offset + 28));
 	return R;
 }
 
