* Added engine test for geometry collection DI (reader+writer)
* Fixed auto-lwc test functionality for spawned actors
* Small DI fixes

#virtualized

diff --git a/Engine/Plugins/Experimental/ChaosNiagara/Source/ChaosNiagara/Private/NiagaraDataInterfaceGeometryCollection.cpp b/Engine/Plugins/Experimental/ChaosNiagara/Source/ChaosNiagara/Private/NiagaraDataInterfaceGeometryCollection.cpp
--- a/Engine/Plugins/Experimental/ChaosNiagara/Source/ChaosNiagara/Private/NiagaraDataInterfaceGeometryCollection.cpp
+++ b/Engine/Plugins/Experimental/ChaosNiagara/Source/ChaosNiagara/Private/NiagaraDataInterfaceGeometryCollection.cpp
@@ -409,9 +409,7 @@ bool UNiagaraDataInterfaceGeometryCollection::InitPerInstanceData(void* PerInsta
 	FNDIGeometryCollectionData* InstanceData = new (PerInstanceData) FNDIGeometryCollectionData();
 
 	check(InstanceData);
-	FNiagaraParameterDirectBinding<UObject*> CollectionParameterBinding;
-	CollectionParameterBinding.Init(SystemInstance->GetInstanceParameters(), GeometryCollectionUserParameter.Parameter);
-	ResolveGeometryCollection(SystemInstance, CollectionParameterBinding.GetValue());
+	ResolveGeometryCollection(SystemInstance);
 	InstanceData->Init(this, SystemInstance);
 	
 	return true;
@@ -437,8 +435,7 @@ void UNiagaraDataInterfaceGeometryCollection::DestroyPerInstanceData(void* PerIn
 	ENQUEUE_RENDER_COMMAND(FNiagaraDIDestroyInstanceData) (
 		[ThisProxy, InstanceID = SystemInstance->GetId()](FRHICommandListImmediate& CmdList)
 		{
-			FNDIGeometryCollectionData* ProxyData =
-				ThisProxy->SystemInstancesToProxyData.Find(InstanceID);
+			FNDIGeometryCollectionData* ProxyData = ThisProxy->SystemInstancesToProxyData.Find(InstanceID);
 
 			if (ProxyData != nullptr && ProxyData->AssetArrays)
 			{			
@@ -452,6 +449,7 @@ void UNiagaraDataInterfaceGeometryCollection::DestroyPerInstanceData(void* PerIn
 bool UNiagaraDataInterfaceGeometryCollection::PerInstanceTick(void* PerInstanceData, FNiagaraSystemInstance* SystemInstance, float InDeltaSeconds)
 {
 	FNDIGeometryCollectionData* InstanceData = static_cast<FNDIGeometryCollectionData*>(PerInstanceData);
+	ResolveGeometryCollection(SystemInstance);
 	if (InstanceData && InstanceData->AssetBuffer && SystemInstance)
 	{
 		InstanceData->Update(this, SystemInstance);
@@ -482,8 +480,12 @@ bool UNiagaraDataInterfaceGeometryCollection::CopyToInternal(UNiagaraDataInterfa
 	return true;
 }
 
-void UNiagaraDataInterfaceGeometryCollection::ResolveGeometryCollection(FNiagaraSystemInstance* SystemInstance, UObject* UserParameter)
+void UNiagaraDataInterfaceGeometryCollection::ResolveGeometryCollection(FNiagaraSystemInstance* SystemInstance)
 {
+	FNiagaraParameterDirectBinding<UObject*> CollectionParameterBinding;
+	CollectionParameterBinding.Init(SystemInstance->GetInstanceParameters(), GeometryCollectionUserParameter.Parameter);
+	UObject* UserParameter = CollectionParameterBinding.GetValue();
+	
 	ResolvedSource = FResolvedNiagaraGeometryCollection();
 
 	switch (SourceMode)
diff --git a/Engine/Plugins/Experimental/ChaosNiagara/Source/ChaosNiagara/Public/NiagaraDataInterfaceGeometryCollection.h b/Engine/Plugins/Experimental/ChaosNiagara/Source/ChaosNiagara/Public/NiagaraDataInterfaceGeometryCollection.h
--- a/Engine/Plugins/Experimental/ChaosNiagara/Source/ChaosNiagara/Public/NiagaraDataInterfaceGeometryCollection.h
+++ b/Engine/Plugins/Experimental/ChaosNiagara/Source/ChaosNiagara/Public/NiagaraDataInterfaceGeometryCollection.h
@@ -212,7 +212,7 @@ class UNiagaraDataInterfaceGeometryCollection : public UNiagaraDataInterface
 #endif
 
 	/** GeometryCollection used to sample from when not overridden by a source actor from the scene. This reference is NOT removed from cooked builds. */
-	UPROPERTY(EditAnywhere, Category = "Geometry Collection", meta = (EditConditionHides, EditCondition = "SourceMode == ENDIGeometryCollection_SourceMode::Default || SourceMode == ENDISocketReaderSourceMode::DefaultCollectionOnly"))
+	UPROPERTY(EditAnywhere, Category = "Geometry Collection", meta = (EditConditionHides, EditCondition = "SourceMode == ENDIGeometryCollection_SourceMode::Default || SourceMode == ENDIGeometryCollection_SourceMode::DefaultCollectionOnly"))
 	TObjectPtr<UGeometryCollection> DefaultGeometryCollection;
 
 	/** The source actor from which to sample. Takes precedence over the direct geometry collection. Note that this can only be set when used as a user variable on a niagara component in the world.*/
@@ -224,7 +224,7 @@ class UNiagaraDataInterfaceGeometryCollection : public UNiagaraDataInterface
 	TObjectPtr<UGeometryCollectionComponent> SourceComponent;
 
 	/** Reference to a user parameter if we're reading one. */
-	UPROPERTY(EditAnywhere, Category = "Geometry Collection", meta = (EditConditionHides, EditCondition = "SourceMode == ENDIGeometryCollection_SourceMode::Default || SourceMode == ENDISocketReaderSourceMode::ParameterBinding"))
+	UPROPERTY(EditAnywhere, Category = "Geometry Collection", meta = (EditConditionHides, EditCondition = "SourceMode == ENDIGeometryCollection_SourceMode::Default || SourceMode == ENDIGeometryCollection_SourceMode::ParameterBinding"))
 	FNiagaraUserParameterBinding GeometryCollectionUserParameter;
 	
 	// If true then this data interface will also read and write intermediate bones or geometry, otherwise only leaf nodes are considered
@@ -283,7 +283,7 @@ protected:
 	virtual bool CopyToInternal(UNiagaraDataInterface* Destination) const override;
 
 private:
-	void ResolveGeometryCollection(FNiagaraSystemInstance* SystemInstance, UObject* UserParameter);
+	void ResolveGeometryCollection(FNiagaraSystemInstance* SystemInstance);
 	bool ResolveGeometryCollectionFromDirectSource();
 	bool ResolveGeometryCollectionFromAttachParent(FNiagaraSystemInstance* SystemInstance);
 	bool ResolveGeometryCollectionFromActor(AActor* Actor);
