Only upload render thread buffer in geometry collection DI when new data is available. Otherwise, multiple ticks per frame cause RDG validation to fail.

diff --git a/Engine/Plugins/Experimental/ChaosNiagara/Source/ChaosNiagara/Private/NiagaraDataInterfaceGeometryCollection.cpp b/Engine/Plugins/Experimental/ChaosNiagara/Source/ChaosNiagara/Private/NiagaraDataInterfaceGeometryCollection.cpp
--- a/Engine/Plugins/Experimental/ChaosNiagara/Source/ChaosNiagara/Private/NiagaraDataInterfaceGeometryCollection.cpp
+++ b/Engine/Plugins/Experimental/ChaosNiagara/Source/ChaosNiagara/Private/NiagaraDataInterfaceGeometryCollection.cpp
@@ -209,8 +209,20 @@ void FNDIGeometryCollectionData::Update(UNiagaraDataInterfaceGeometryCollection*
 			if (NumPieces != AssetArrays->BoundsBuffer.Num())
 			{
 				Init(Interface, SystemInstance);
+				bNeedsRenderUpdate = true;
+				AssetArrays->ComponentRestTransformBuffer = Interface->ResolvedSource.GetInitialLocalRestTransforms();
+			}
+			else
+			{
+				TArray<FTransform> NewTransforms = Interface->ResolvedSource.GetInitialLocalRestTransforms();
+				int32 TransformCount = NewTransforms.Num();
+				if (TransformCount != AssetArrays->ComponentRestTransformBuffer.Num() ||
+					FMemory::Memcmp(NewTransforms.GetData(), AssetArrays->ComponentRestTransformBuffer.GetData(), TransformCount * sizeof(FTransform)) != 0)
+				{
+					AssetArrays->ComponentRestTransformBuffer = MoveTemp(NewTransforms);
+					bNeedsRenderUpdate = true;
+				}
 			}
-			AssetArrays->ComponentRestTransformBuffer = Interface->ResolvedSource.GetInitialLocalRestTransforms();
 			
 			FVector Origin(ForceInitToZero);
 			FVector Extents(ForceInitToZero);
@@ -1058,7 +1070,7 @@ void UNiagaraDataInterfaceGeometryCollection::ProvidePerInstanceDataForRenderThr
 	FNDIGeometryCollectionData* GameThreadData = static_cast<FNDIGeometryCollectionData*>(PerInstanceData);
 	FNDIGeometryCollectionData* RenderThreadData = static_cast<FNDIGeometryCollectionData*>(DataForRenderThread);
 
-	if (GameThreadData != nullptr && RenderThreadData != nullptr)
+	if (GameThreadData && RenderThreadData)
 	{		
 		RenderThreadData->AssetBuffer = GameThreadData->AssetBuffer;				
 
@@ -1069,22 +1081,27 @@ void UNiagaraDataInterfaceGeometryCollection::ProvidePerInstanceDataForRenderThr
 		RenderThreadData->BoundsExtent = GameThreadData->BoundsExtent;
 		RenderThreadData->RootTransform = GameThreadData->RootTransform;
 
-		constexpr int32 TransformGpuSize = 10 * 4;
-		const int32 BufferSize = RenderThreadData->AssetArrays->ElementIndexToTransformBufferMapping.Num() * TransformGpuSize;
-		RenderThreadData->AssetBuffer->DataToUpload.SetNumUninitialized(BufferSize);
-		float* OutFloats = reinterpret_cast<float*>(RenderThreadData->AssetBuffer->DataToUpload.GetData());
-		for (int32 Index : RenderThreadData->AssetArrays->ElementIndexToTransformBufferMapping)
+		if (GameThreadData->bNeedsRenderUpdate)
 		{
-			const FTransform& Transform = RenderThreadData->AssetArrays->ComponentRestTransformBuffer[Index];
+			GameThreadData->bNeedsRenderUpdate = false;
 			
-			const FVector3f Translation(Transform.GetTranslation());
-			const FQuat4f Rotation(Transform.GetRotation());
-			const FVector3f Scale(Transform.GetScale3D());
-
-			FMemory::Memcpy(&OutFloats[0], &Translation, sizeof(FVector3f));
-			FMemory::Memcpy(&OutFloats[3], &Rotation, sizeof(FQuat4f));
-			FMemory::Memcpy(&OutFloats[7], &Scale, sizeof(FVector3f));
-			OutFloats += 10;
+			constexpr int32 TransformGpuSize = 10 * 4;
+			const int32 BufferSize = RenderThreadData->AssetArrays->ElementIndexToTransformBufferMapping.Num() * TransformGpuSize;
+			RenderThreadData->AssetBuffer->DataToUpload.SetNumUninitialized(BufferSize);
+			float* OutFloats = reinterpret_cast<float*>(RenderThreadData->AssetBuffer->DataToUpload.GetData());
+			for (int32 Index : RenderThreadData->AssetArrays->ElementIndexToTransformBufferMapping)
+			{
+				const FTransform& Transform = RenderThreadData->AssetArrays->ComponentRestTransformBuffer[Index];
+			
+				const FVector3f Translation(Transform.GetTranslation());
+				const FQuat4f Rotation(Transform.GetRotation());
+				const FVector3f Scale(Transform.GetScale3D());
+
+				FMemory::Memcpy(&OutFloats[0], &Translation, sizeof(FVector3f));
+				FMemory::Memcpy(&OutFloats[3], &Rotation, sizeof(FQuat4f));
+				FMemory::Memcpy(&OutFloats[7], &Scale, sizeof(FVector3f));
+				OutFloats += 10;
+			}
 		}
 	}
 	check(Proxy);
diff --git a/Engine/Plugins/Experimental/ChaosNiagara/Source/ChaosNiagara/Public/NiagaraDataInterfaceGeometryCollection.h b/Engine/Plugins/Experimental/ChaosNiagara/Source/ChaosNiagara/Public/NiagaraDataInterfaceGeometryCollection.h
--- a/Engine/Plugins/Experimental/ChaosNiagara/Source/ChaosNiagara/Public/NiagaraDataInterfaceGeometryCollection.h
+++ b/Engine/Plugins/Experimental/ChaosNiagara/Source/ChaosNiagara/Public/NiagaraDataInterfaceGeometryCollection.h
@@ -135,7 +135,11 @@ struct FNDIGeometryCollectionData
 	/** Physics asset Cpu arrays */
 	FNDIGeometryCollectionArrays* AssetArrays = nullptr;
 
+	// Flag when there pending transform writes that need to go back to the component
 	bool bHasPendingComponentTransformUpdate = false;
+
+	// True if we need to upload new data to the gpu
+	bool bNeedsRenderUpdate = false;
 };
 
 UENUM()
