Fixed a crash when using geometry collection DI with certain project settings

#tests Tested with flip 3d fluid in EngineTest and ContentExamples project, as well as running automated tests. To trigger the crash in a fluid system, the use of the geometry collection needs to be enabled in the emitter properties via static bool.
#rnx

diff --git a/Engine/Plugins/Experimental/ChaosNiagara/Source/ChaosNiagara/Private/NiagaraDataInterfaceGeometryCollection.cpp b/Engine/Plugins/Experimental/ChaosNiagara/Source/ChaosNiagara/Private/NiagaraDataInterfaceGeometryCollection.cpp
--- a/Engine/Plugins/Experimental/ChaosNiagara/Source/ChaosNiagara/Private/NiagaraDataInterfaceGeometryCollection.cpp
+++ b/Engine/Plugins/Experimental/ChaosNiagara/Source/ChaosNiagara/Private/NiagaraDataInterfaceGeometryCollection.cpp
@@ -384,11 +384,16 @@ void FNDIGeometryCollectionProxy::PreStage(const FNDIGpuComputePreStageContext&
 			NDIGeometryCollectionLocal::UpdateInternalBuffer<FVector4f>(RHICmdList, ProxyData->AssetArrays->BoundsBuffer, ProxyData->AssetBuffer->BoundsBuffer);
 
 			// build RDG buffer with per-element component transforms
+			FRDGBuilder& GraphBuilder = Context.GetGraphBuilder();
 			TArray<uint8>& DataToUpload = ProxyData->AssetBuffer->DataToUpload;
-			if (DataToUpload.Num() > 0)
+			if (DataToUpload.IsEmpty() && !ProxyData->AssetBuffer->ComponentRestTransformBuffer.IsValid())
 			{
-				FRDGBuilder& GraphBuilder = Context.GetGraphBuilder();
+				// add dummy data so we can bind something to the buffer
+				DataToUpload.SetNumZeroed(12);
+			}
 
+			if (!DataToUpload.IsEmpty())
+			{
 				const uint64 InitialDataSize = DataToUpload.Num() * DataToUpload.GetTypeSize();
 				const uint64 BufferSize = Align(InitialDataSize, 16);
 
@@ -1045,10 +1050,6 @@ void UNiagaraDataInterfaceGeometryCollection::SetShaderParameters(const FNiagara
 {
 	FNDIGeometryCollectionProxy& InterfaceProxy = Context.GetProxy<FNDIGeometryCollectionProxy>();
 	FNDIGeometryCollectionData& ProxyData = InterfaceProxy.SystemInstancesToProxyData.FindChecked(Context.GetSystemInstanceID());
-	
-	FRDGBuilder& GraphBuilder = Context.GetGraphBuilder();
-	FRDGBufferRef RDGBuffer = GraphBuilder.RegisterExternalBuffer(ProxyData.AssetBuffer->ComponentRestTransformBuffer);
-	FRDGBufferSRVRef RDGBufferSRV = GraphBuilder.CreateSRV(RDGBuffer);
 
 	FShaderParameters* ShaderParameters = Context.GetParameterNestedStruct<FShaderParameters>();
 	ShaderParameters->BoundsMin							= ProxyData.BoundsOrigin - ProxyData.BoundsExtent;
@@ -1062,7 +1063,13 @@ void UNiagaraDataInterfaceGeometryCollection::SetShaderParameters(const FNiagara
 	ShaderParameters->WorldInverseTransformBuffer		= FNiagaraRenderer::GetSrvOrDefaultFloat4(ProxyData.AssetBuffer->WorldInverseTransformBuffer.SRV);
 	ShaderParameters->PrevWorldInverseTransformBuffer	= FNiagaraRenderer::GetSrvOrDefaultFloat4(ProxyData.AssetBuffer->PrevWorldInverseTransformBuffer.SRV);
 	ShaderParameters->BoundsBuffer						= FNiagaraRenderer::GetSrvOrDefaultFloat4(ProxyData.AssetBuffer->BoundsBuffer.SRV);
-	ShaderParameters->ElementTransforms					= RDGBufferSRV;
+
+	if (ensure(ProxyData.AssetBuffer->ComponentRestTransformBuffer))
+	{
+		FRDGBuilder& GraphBuilder = Context.GetGraphBuilder();
+		FRDGBufferRef RDGBuffer = GraphBuilder.RegisterExternalBuffer(ProxyData.AssetBuffer->ComponentRestTransformBuffer);
+		ShaderParameters->ElementTransforms	= GraphBuilder.CreateSRV(RDGBuffer);
+	}
 }
 
 void UNiagaraDataInterfaceGeometryCollection::ProvidePerInstanceDataForRenderThread(void* DataForRenderThread, void* PerInstanceData, const FNiagaraSystemInstanceID& SystemInstance)
