Fix potential crash in niagara chaos destruction interface
- Happens when a geometry collection component is destroyed while being tracked for trailing
- Made sure we use weakptr for tracking the component
- Only add the component to the list if it tracks for trailing
- Rename the array to be clear about its intended usage

diff --git a/Engine/Plugins/Experimental/ChaosNiagara/Source/ChaosNiagara/Classes/NiagaraDataInterfaceChaosDestruction.h b/Engine/Plugins/Experimental/ChaosNiagara/Source/ChaosNiagara/Classes/NiagaraDataInterfaceChaosDestruction.h
--- a/Engine/Plugins/Experimental/ChaosNiagara/Source/ChaosNiagara/Classes/NiagaraDataInterfaceChaosDestruction.h
+++ b/Engine/Plugins/Experimental/ChaosNiagara/Source/ChaosNiagara/Classes/NiagaraDataInterfaceChaosDestruction.h
@@ -766,8 +766,8 @@ protected:
 	TArray<Chaos::FBreakingDataExt> BreakingEvents;
 	TArray<Chaos::FTrailingDataExt> TrailingEvents;
 
-	// Saving GeometryCollectionComponents from breaking for trailing
-	TSet<UGeometryCollectionComponent*> GeometryCollectionComponentsFromBreaking;
+	// geometry collection to track for trailing 
+	TSet<TWeakObjectPtr<UGeometryCollectionComponent>> TrailingGeometryCollectionComponents;
 };
 
 struct FNiagaraDataInterfaceProxyChaosDestruction : public FNiagaraDataInterfaceProxy
diff --git a/Engine/Plugins/Experimental/ChaosNiagara/Source/ChaosNiagara/Private/NiagaraDataInterfaceChaosDestruction.cpp b/Engine/Plugins/Experimental/ChaosNiagara/Source/ChaosNiagara/Private/NiagaraDataInterfaceChaosDestruction.cpp
--- a/Engine/Plugins/Experimental/ChaosNiagara/Source/ChaosNiagara/Private/NiagaraDataInterfaceChaosDestruction.cpp
+++ b/Engine/Plugins/Experimental/ChaosNiagara/Source/ChaosNiagara/Private/NiagaraDataInterfaceChaosDestruction.cpp
@@ -1500,8 +1500,10 @@ void UNiagaraDataInterfaceChaosDestruction::HandleBreakingEvents(const Chaos::FB
 					CopyData.PhysicalMaterialName = FName();
 				}
 
-				// Save GeometryCollectionComponent for trailing
-				GeometryCollectionComponentsFromBreaking.Add(GeometryCollectionComponent);
+				if (GeometryCollectionComponent->GetNotifyTrailing())
+				{
+					TrailingGeometryCollectionComponents.Add(GeometryCollectionComponent);
+				}
 			}
 		}
 		else
@@ -2024,88 +2026,91 @@ void UNiagaraDataInterfaceChaosDestruction::HandleTrailingEvents(const Chaos::FT
 		return false;
 	};
 
-	if (GeometryCollectionComponentsFromBreaking.Num() > 0)
+	TArray<TWeakObjectPtr<UGeometryCollectionComponent>> ComponentsToRemove;
+	for (const TWeakObjectPtr<UGeometryCollectionComponent>& WeakGeometryCollectionComponent: TrailingGeometryCollectionComponents)
 	{
-		TArray<UGeometryCollectionComponent*> ComponentsToRemove;
-		for (auto& GeometryCollectionComponent : GeometryCollectionComponentsFromBreaking)
+		if (!WeakGeometryCollectionComponent.IsValid())
 		{
-			if (GeometryCollectionComponent)
+			ComponentsToRemove.Add(WeakGeometryCollectionComponent);
+			continue;
+		}
+
+		if (const TStrongObjectPtr<UGeometryCollectionComponent> GeometryCollectionComponent = WeakGeometryCollectionComponent.Pin())
+		{
+			if (!GeometryCollectionComponent->IsPhysicsStateCreated())
 			{
-				if (!GeometryCollectionComponent->IsPhysicsStateCreated())
-				{
-					ComponentsToRemove.Add(GeometryCollectionComponent);
-				}
-				else if (GeometryCollectionComponent->GetNotifyTrailing())
+				ComponentsToRemove.Add(WeakGeometryCollectionComponent);
+			}
+			else if (GeometryCollectionComponent->GetNotifyTrailing())
+			{
+				PhysicalMaterial = GeometryCollectionComponent->GetPhysicalMaterial();
+				ensure(PhysicalMaterial);
+				if (PhysicalMaterial)
 				{
-					PhysicalMaterial = GeometryCollectionComponent->GetPhysicalMaterial();
-					ensure(PhysicalMaterial);
-					if (PhysicalMaterial)
+					if (bApplyMaterialsFilter && !IsMaterialInFilter(PhysicalMaterial->GetFName()))
 					{
-						if (bApplyMaterialsFilter && !IsMaterialInFilter(PhysicalMaterial->GetFName()))
-						{
-							continue;
-						}
+						continue;
 					}
+				}
 
-					const TArray<FTransform3f>& ComponentSpaceTransforms = GeometryCollectionComponent->GetComponentSpaceTransforms3f();
-					const FTransform ActorTransform = GeometryCollectionComponent->GetComponentToWorld();
+				const TArray<FTransform3f>& ComponentSpaceTransforms = GeometryCollectionComponent->GetComponentSpaceTransforms3f();
+				const FTransform ActorTransform = GeometryCollectionComponent->GetComponentToWorld();
 
-					const FGeometryDynamicCollection* DynamicCollection = GeometryCollectionComponent->GetDynamicCollection();
+				const FGeometryDynamicCollection* DynamicCollection = GeometryCollectionComponent->GetDynamicCollection();
 
-					if (DynamicCollection)
+				if (DynamicCollection)
+				{
+					if (const UGeometryCollection* RestCollection = GeometryCollectionComponent->GetRestCollection())
 					{
-						if (const UGeometryCollection* RestCollection = GeometryCollectionComponent->GetRestCollection())
-						{
-							const TSharedPtr<FGeometryCollection, ESPMode::ThreadSafe> GeometryCollection = RestCollection->GetGeometryCollection();
+						const TSharedPtr<FGeometryCollection, ESPMode::ThreadSafe> GeometryCollection = RestCollection->GetGeometryCollection();
 
-							// Get the MassToLocal transforms
-							if (GeometryCollection->HasAttribute(TEXT("MassToLocal"), FTransformCollection::TransformGroup))
-							{
-								const TManagedArray<FTransform>& CollectionMassToLocal = GeometryCollection->GetAttribute<FTransform>(TEXT("MassToLocal"), FTransformCollection::TransformGroup);
+						// Get the MassToLocal transforms
+						if (GeometryCollection->HasAttribute(TEXT("MassToLocal"), FTransformCollection::TransformGroup))
+						{
+							const TManagedArray<FTransform>& CollectionMassToLocal = GeometryCollection->GetAttribute<FTransform>(TEXT("MassToLocal"), FTransformCollection::TransformGroup);
 
-								const TManagedArray<FVector3f>* LinearVelocity = DynamicCollection->FindAttributeTyped<FVector3f>("LinearVelocity", FTransformCollection::TransformGroup);
-								const TManagedArray<FVector3f>* AngularVelocity = DynamicCollection->FindAttributeTyped<FVector3f>("AngularVelocity", FTransformCollection::TransformGroup);
+							const TManagedArray<FVector3f>* LinearVelocity = DynamicCollection->FindAttributeTyped<FVector3f>("LinearVelocity", FTransformCollection::TransformGroup);
+							const TManagedArray<FVector3f>* AngularVelocity = DynamicCollection->FindAttributeTyped<FVector3f>("AngularVelocity", FTransformCollection::TransformGroup);
 
-								if (!LinearVelocity || !AngularVelocity)
-								{
-									continue;
-								}
+							if (!LinearVelocity || !AngularVelocity)
+							{
+								continue;
+							}
 
-								ensure(DynamicCollection->Active.Num() == (*LinearVelocity).Num());
-								ensure(DynamicCollection->Active.Num() == (*AngularVelocity).Num());
+							ensure(DynamicCollection->Active.Num() == (*LinearVelocity).Num());
+							ensure(DynamicCollection->Active.Num() == (*AngularVelocity).Num());
 
-								for (int Idx = 0; Idx < DynamicCollection->Active.Num(); ++Idx)
+							for (int Idx = 0; Idx < DynamicCollection->Active.Num(); ++Idx)
+							{
+								if (DynamicCollection->Active[Idx] && (*LinearVelocity)[Idx].SquaredLength() >= TrailMinSpeedToSpawnSquared)
 								{
-									if (DynamicCollection->Active[Idx] && (*LinearVelocity)[Idx].SquaredLength() >= TrailMinSpeedToSpawnSquared)
-									{
-										Chaos::FTrailingDataExt TrailingData;
+									Chaos::FTrailingDataExt TrailingData;
 
-										const FTransform CurrTransform = CollectionMassToLocal[Idx] * FTransform(ComponentSpaceTransforms[Idx]) * ActorTransform;
-										TrailingData.Location = CurrTransform.GetTranslation();
+									const FTransform CurrTransform = CollectionMassToLocal[Idx] * FTransform(ComponentSpaceTransforms[Idx]) * ActorTransform;
+									TrailingData.Location = CurrTransform.GetTranslation();
 
-										TrailingData.Velocity = (*LinearVelocity)[Idx];
-										TrailingData.AngularVelocity = (*AngularVelocity)[Idx];
+									TrailingData.Velocity = (*LinearVelocity)[Idx];
+									TrailingData.AngularVelocity = (*AngularVelocity)[Idx];
 
-										TrailingData.Mass = 1.f;
-										TrailingData.BoundingboxVolume = 1000000.f;
-										TrailingData.BoundingboxExtentMin = 100.f;
-										TrailingData.BoundingboxExtentMax = 100.f;
-										TrailingData.SurfaceType = 0;
+									TrailingData.Mass = 1.f;
+									TrailingData.BoundingboxVolume = 1000000.f;
+									TrailingData.BoundingboxExtentMin = 100.f;
+									TrailingData.BoundingboxExtentMax = 100.f;
+									TrailingData.SurfaceType = 0;
 
-										TrailingEvents.Add(TrailingData);
-									}
+									TrailingEvents.Add(TrailingData);
 								}
 							}
 						}
-					}				
-				}
+					}
+				}				
 			}
 		}
+	}
 
-		for (UGeometryCollectionComponent* ComponentToRemove : ComponentsToRemove)
-		{
-			GeometryCollectionComponentsFromBreaking.Remove(ComponentToRemove);
-		}
+	for (const TWeakObjectPtr<UGeometryCollectionComponent>& ComponentToRemove : ComponentsToRemove)
+	{
+		TrailingGeometryCollectionComponents.Remove(ComponentToRemove);
 	}
 }
 
