Fixed a bug where multiple instances of a system could write to the same geometry collection component

#tests Checked that engine tests are green and ran a test level with multiple collections

diff --git a/Engine/Plugins/Experimental/ChaosNiagara/Source/ChaosNiagara/Private/NiagaraDataInterfaceGeometryCollection.cpp b/Engine/Plugins/Experimental/ChaosNiagara/Source/ChaosNiagara/Private/NiagaraDataInterfaceGeometryCollection.cpp
--- a/Engine/Plugins/Experimental/ChaosNiagara/Source/ChaosNiagara/Private/NiagaraDataInterfaceGeometryCollection.cpp
+++ b/Engine/Plugins/Experimental/ChaosNiagara/Source/ChaosNiagara/Private/NiagaraDataInterfaceGeometryCollection.cpp
@@ -118,7 +118,7 @@ void FNDIGeometryCollectionData::Init(UNiagaraDataInterfaceGeometryCollection* I
 
 	if (Interface && SystemInstance)
 	{
-		if (const UGeometryCollection* GeometryCollection = Interface->ResolvedSource.GetGeometryCollection())
+		if (const UGeometryCollection* GeometryCollection = ResolvedSource.GetGeometryCollection())
 		{			
 			const TSharedPtr<FGeometryCollection> Collection = GeometryCollection->GetGeometryCollection();
 			const TManagedArray<FBox>& BoundingBoxes = Collection->BoundingBox;
@@ -144,9 +144,9 @@ void FNDIGeometryCollectionData::Init(UNiagaraDataInterfaceGeometryCollection* I
 
 			FVector Origin(ForceInitToZero);
 			FVector Extents(ForceInitToZero);
-			if (Interface->ResolvedSource.Component.IsValid())
+			if (ResolvedSource.Component.IsValid())
 			{
-				Interface->ResolvedSource.Component->Bounds.GetBox().GetCenterAndExtents(Origin, Extents);
+				ResolvedSource.Component->Bounds.GetBox().GetCenterAndExtents(Origin, Extents);
 			}
 
 			FNiagaraLWCConverter LwcConverter = SystemInstance->GetLWCConverter();
@@ -187,9 +187,9 @@ void FNDIGeometryCollectionData::Update(UNiagaraDataInterfaceGeometryCollection*
 	{		
 		TickingGroup = ComputeTickingGroup();
 		
-		if (const UGeometryCollection* GeoCollection = Interface->ResolvedSource.GetGeometryCollection())
+		if (const UGeometryCollection* GeoCollection = ResolvedSource.GetGeometryCollection())
 		{
-			RootTransform = Interface->ResolvedSource.GetComponentRootTransform(SystemInstance);
+			RootTransform = ResolvedSource.GetComponentRootTransform(SystemInstance);
 
 			const TSharedPtr<FGeometryCollection> Collection = GeoCollection->GetGeometryCollection();
 			const TManagedArray<FBox>& BoundingBoxes = Collection->BoundingBox;
@@ -210,11 +210,11 @@ void FNDIGeometryCollectionData::Update(UNiagaraDataInterfaceGeometryCollection*
 			{
 				Init(Interface, SystemInstance);
 				bNeedsRenderUpdate = true;
-				AssetArrays->ComponentRestTransformBuffer = Interface->ResolvedSource.GetInitialLocalRestTransforms();
+				AssetArrays->ComponentRestTransformBuffer = ResolvedSource.GetInitialLocalRestTransforms();
 			}
 			else
 			{
-				TArray<FTransform> NewTransforms = Interface->ResolvedSource.GetInitialLocalRestTransforms();
+				TArray<FTransform> NewTransforms = ResolvedSource.GetInitialLocalRestTransforms();
 				int32 TransformCount = NewTransforms.Num();
 				if (TransformCount != AssetArrays->ComponentRestTransformBuffer.Num() ||
 					FMemory::Memcmp(NewTransforms.GetData(), AssetArrays->ComponentRestTransformBuffer.GetData(), TransformCount * sizeof(FTransform)) != 0)
@@ -226,9 +226,9 @@ void FNDIGeometryCollectionData::Update(UNiagaraDataInterfaceGeometryCollection*
 			
 			FVector Origin(ForceInitToZero);
 			FVector Extents(ForceInitToZero);
-			if (Interface->ResolvedSource.Component.IsValid())
+			if (ResolvedSource.Component.IsValid())
 			{
-				Interface->ResolvedSource.Component->Bounds.GetBox().GetCenterAndExtents(Origin, Extents);
+				ResolvedSource.Component->Bounds.GetBox().GetCenterAndExtents(Origin, Extents);
 			}
 
 			FNiagaraLWCConverter LwcConverter = SystemInstance->GetLWCConverter();
@@ -259,7 +259,7 @@ void FNDIGeometryCollectionData::Update(UNiagaraDataInterfaceGeometryCollection*
 					FVector LocalTranslation = (CurrBox.Max + CurrBox.Min) * .5;
 					FTransform LocalOffset(LocalTranslation);
 
-					const FTransform3f CurrTransform(LocalOffset * Interface->ResolvedSource.GetComponentSpaceTransform(CurrTransformIndex) * RootTransform);
+					const FTransform3f CurrTransform(LocalOffset * ResolvedSource.GetComponentSpaceTransform(CurrTransformIndex) * RootTransform);
 					CurrTransform.ToMatrixWithScale().To3x4MatrixTranspose(&AssetArrays->WorldTransformBuffer[TransformIndex].X);
 
 					const FTransform3f CurrInverse = CurrTransform.Inverse();
@@ -426,7 +426,7 @@ bool UNiagaraDataInterfaceGeometryCollection::InitPerInstanceData(void* PerInsta
 	FNDIGeometryCollectionData* InstanceData = new (PerInstanceData) FNDIGeometryCollectionData();
 
 	check(InstanceData);
-	ResolveGeometryCollection(SystemInstance);
+	ResolveGeometryCollection(SystemInstance, InstanceData);
 	InstanceData->Init(this, SystemInstance);
 	
 	return true;
@@ -466,7 +466,7 @@ void UNiagaraDataInterfaceGeometryCollection::DestroyPerInstanceData(void* PerIn
 bool UNiagaraDataInterfaceGeometryCollection::PerInstanceTick(void* PerInstanceData, FNiagaraSystemInstance* SystemInstance, float InDeltaSeconds)
 {
 	FNDIGeometryCollectionData* InstanceData = static_cast<FNDIGeometryCollectionData*>(PerInstanceData);
-	ResolveGeometryCollection(SystemInstance);
+	ResolveGeometryCollection(SystemInstance, InstanceData);
 	if (InstanceData && InstanceData->AssetBuffer && SystemInstance)
 	{
 		InstanceData->Update(this, SystemInstance);
@@ -492,42 +492,41 @@ bool UNiagaraDataInterfaceGeometryCollection::CopyToInternal(UNiagaraDataInterfa
 	OtherTyped->SourceComponent = SourceComponent;
 	OtherTyped->GeometryCollectionUserParameter = GeometryCollectionUserParameter;
 	OtherTyped->bIncludeIntermediateBones = bIncludeIntermediateBones;
-	OtherTyped->ResolvedSource = ResolvedSource;
 
 	return true;
 }
 
-void UNiagaraDataInterfaceGeometryCollection::ResolveGeometryCollection(FNiagaraSystemInstance* SystemInstance)
+void UNiagaraDataInterfaceGeometryCollection::ResolveGeometryCollection(FNiagaraSystemInstance* SystemInstance, FNDIGeometryCollectionData* InstanceData)
 {
 	FNiagaraParameterDirectBinding<UObject*> CollectionParameterBinding;
 	CollectionParameterBinding.Init(SystemInstance->GetInstanceParameters(), GeometryCollectionUserParameter.Parameter);
 	UObject* UserParameter = CollectionParameterBinding.GetValue();
 	
-	ResolvedSource = FResolvedNiagaraGeometryCollection();
+	InstanceData->ResolvedSource = FResolvedNiagaraGeometryCollection();
 
 	switch (SourceMode)
 	{
 	case ENDIGeometryCollection_SourceMode::Source:
-		ResolveGeometryCollectionFromDirectSource();
+		ResolveGeometryCollectionFromDirectSource(InstanceData->ResolvedSource);
 		break;
 	case ENDIGeometryCollection_SourceMode::AttachParent:
-		ResolveGeometryCollectionFromAttachParent(SystemInstance);
+		ResolveGeometryCollectionFromAttachParent(SystemInstance, InstanceData->ResolvedSource);
 		break;
 	case ENDIGeometryCollection_SourceMode::DefaultCollectionOnly:
-		ResolveGeometryCollectionFromDefaultCollection();
+		ResolveGeometryCollectionFromDefaultCollection(InstanceData->ResolvedSource);
 		break;
 	case ENDIGeometryCollection_SourceMode::ParameterBinding:
-		ResolveGeometryCollectionFromParamterBinding(UserParameter);
+		ResolveGeometryCollectionFromParameterBinding(UserParameter, InstanceData->ResolvedSource);
 		break;
 	case ENDIGeometryCollection_SourceMode::Default:
 	default:
-		if (!ResolveGeometryCollectionFromDirectSource())
+		if (!ResolveGeometryCollectionFromDirectSource(InstanceData->ResolvedSource))
 		{
-			if (!ResolveGeometryCollectionFromParamterBinding(UserParameter))
+			if (!ResolveGeometryCollectionFromParameterBinding(UserParameter, InstanceData->ResolvedSource))
 			{
-				if (!ResolveGeometryCollectionFromAttachParent(SystemInstance))
+				if (!ResolveGeometryCollectionFromAttachParent(SystemInstance, InstanceData->ResolvedSource))
 				{
-					ResolveGeometryCollectionFromDefaultCollection();
+					ResolveGeometryCollectionFromDefaultCollection(InstanceData->ResolvedSource);
 				}
 			}
 		}
@@ -535,15 +534,15 @@ void UNiagaraDataInterfaceGeometryCollection::ResolveGeometryCollection(FNiagara
 	}
 
 #if WITH_EDITORONLY_DATA
-	if (!ResolvedSource.Collection.IsValid() && !ResolvedSource.Component.IsValid() && (!SystemInstance || !SystemInstance->GetWorld()->IsGameWorld()))
+	if (!InstanceData->ResolvedSource.Collection.IsValid() && !InstanceData->ResolvedSource.Component.IsValid() && (!SystemInstance || !SystemInstance->GetWorld()->IsGameWorld()))
 	{
 		// NOTE: We don't fall back on the preview mesh if we have a valid collection referenced
-		ResolvedSource.Collection = PreviewCollection.LoadSynchronous();
+		InstanceData->ResolvedSource.Collection = PreviewCollection.LoadSynchronous();
 	}
 #endif
 }
 
-bool UNiagaraDataInterfaceGeometryCollection::ResolveGeometryCollectionFromDirectSource()
+bool UNiagaraDataInterfaceGeometryCollection::ResolveGeometryCollectionFromDirectSource(FResolvedNiagaraGeometryCollection& ResolvedSource)
 {
 	if (::IsValid(SourceComponent))
 	{
@@ -557,7 +556,7 @@ bool UNiagaraDataInterfaceGeometryCollection::ResolveGeometryCollectionFromDirec
 	return false;
 }
 
-bool UNiagaraDataInterfaceGeometryCollection::ResolveGeometryCollectionFromAttachParent(FNiagaraSystemInstance* SystemInstance)
+bool UNiagaraDataInterfaceGeometryCollection::ResolveGeometryCollectionFromAttachParent(FNiagaraSystemInstance* SystemInstance, FResolvedNiagaraGeometryCollection& ResolvedSource)
 {
 	if (USceneComponent* AttachComponent = SystemInstance->GetAttachComponent())
 	{
@@ -579,12 +578,12 @@ bool UNiagaraDataInterfaceGeometryCollection::ResolveGeometryCollectionFromAttac
 			ResolvedSource.Component = OuterComp;
 			return true;
 		}
-		return ResolveGeometryCollectionFromActor(AttachComponent->GetAttachmentRootActor());
+		return ResolveGeometryCollectionFromActor(AttachComponent->GetAttachmentRootActor(), ResolvedSource);
 	}
 	return false;
 }
 
-bool UNiagaraDataInterfaceGeometryCollection::ResolveGeometryCollectionFromActor(AActor* Actor)
+bool UNiagaraDataInterfaceGeometryCollection::ResolveGeometryCollectionFromActor(AActor* Actor, FResolvedNiagaraGeometryCollection& ResolvedSource)
 {
 	if (Actor)
 	{
@@ -609,19 +608,19 @@ bool UNiagaraDataInterfaceGeometryCollection::ResolveGeometryCollectionFromActor
 	return false;
 }
 
-bool UNiagaraDataInterfaceGeometryCollection::ResolveGeometryCollectionFromDefaultCollection()
+bool UNiagaraDataInterfaceGeometryCollection::ResolveGeometryCollectionFromDefaultCollection(FResolvedNiagaraGeometryCollection& ResolvedSource)
 {
 	ResolvedSource.Collection = DefaultGeometryCollection;
 	return true;
 }
 
-bool UNiagaraDataInterfaceGeometryCollection::ResolveGeometryCollectionFromParamterBinding(UObject* ParameterBindingValue)
+bool UNiagaraDataInterfaceGeometryCollection::ResolveGeometryCollectionFromParameterBinding(UObject* ParameterBindingValue, FResolvedNiagaraGeometryCollection& ResolvedSource)
 {
 	if (ParameterBindingValue == nullptr)
 	{
 		return false;
 	}
-	if (ResolveGeometryCollectionFromActor(Cast<AActor>(ParameterBindingValue)))
+	if (ResolveGeometryCollectionFromActor(Cast<AActor>(ParameterBindingValue), ResolvedSource))
 	{
 		return true;
 	}
@@ -641,9 +640,9 @@ bool UNiagaraDataInterfaceGeometryCollection::ResolveGeometryCollectionFromParam
 bool UNiagaraDataInterfaceGeometryCollection::PerInstanceTickPostSimulate(void* PerInstanceData, FNiagaraSystemInstance* InSystemInstance, float DeltaSeconds)
 {
 	FNDIGeometryCollectionData* InstanceData = static_cast<FNDIGeometryCollectionData*>(PerInstanceData);
-	if (InstanceData && InstanceData->bHasPendingComponentTransformUpdate && ResolvedSource.Component.IsValid())
+	if (InstanceData && InstanceData->bHasPendingComponentTransformUpdate && InstanceData->ResolvedSource.Component.IsValid())
 	{
-		ResolvedSource.Component->SetLocalRestTransforms(InstanceData->AssetArrays->ComponentRestTransformBuffer, !bIncludeIntermediateBones);
+		InstanceData->ResolvedSource.Component->SetLocalRestTransforms(InstanceData->AssetArrays->ComponentRestTransformBuffer, !bIncludeIntermediateBones);
 	}
 	return false;
 }
diff --git a/Engine/Plugins/Experimental/ChaosNiagara/Source/ChaosNiagara/Public/NiagaraDataInterfaceGeometryCollection.h b/Engine/Plugins/Experimental/ChaosNiagara/Source/ChaosNiagara/Public/NiagaraDataInterfaceGeometryCollection.h
--- a/Engine/Plugins/Experimental/ChaosNiagara/Source/ChaosNiagara/Public/NiagaraDataInterfaceGeometryCollection.h
+++ b/Engine/Plugins/Experimental/ChaosNiagara/Source/ChaosNiagara/Public/NiagaraDataInterfaceGeometryCollection.h
@@ -104,6 +104,17 @@ struct FNDIGeometryCollectionBuffer : public FRenderResource
 	}
 };
 
+struct FResolvedNiagaraGeometryCollection
+{
+	const UGeometryCollection* GetGeometryCollection() const;
+	FTransform GetComponentRootTransform(FNiagaraSystemInstance* SystemInstance) const;
+	FTransform GetComponentSpaceTransform(int32 TransformIndex) const;
+	TArray<FTransform> GetInitialLocalRestTransforms() const;
+	
+	TWeakObjectPtr<UGeometryCollection> Collection;
+	TWeakObjectPtr<UGeometryCollectionComponent> Component;
+};
+
 /** Data stored per physics asset instance*/
 struct FNDIGeometryCollectionData
 {
@@ -140,6 +151,8 @@ struct FNDIGeometryCollectionData
 
 	// True if we need to upload new data to the gpu
 	bool bNeedsRenderUpdate = false;
+
+	FResolvedNiagaraGeometryCollection ResolvedSource;
 };
 
 UENUM()
@@ -167,23 +180,6 @@ enum class ENDIGeometryCollection_SourceMode : uint8
 	ParameterBinding,
 };
 
-USTRUCT()
-struct FResolvedNiagaraGeometryCollection
-{
-	GENERATED_BODY()
-	
-	const UGeometryCollection* GetGeometryCollection() const;
-	FTransform GetComponentRootTransform(FNiagaraSystemInstance* SystemInstance) const;
-	FTransform GetComponentSpaceTransform(int32 TransformIndex) const;
-	TArray<FTransform> GetInitialLocalRestTransforms() const;
-	
-	UPROPERTY()
-	TWeakObjectPtr<UGeometryCollection> Collection;
-
-	UPROPERTY()
-	TWeakObjectPtr<UGeometryCollectionComponent> Component;
-};
-
 /** Data Interface for the Collisions */
 UCLASS(EditInlineNew, Category = "Chaos", meta = (DisplayName = "Geometry Collection"), MinimalAPI)
 class UNiagaraDataInterfaceGeometryCollection : public UNiagaraDataInterface
@@ -235,9 +231,6 @@ class UNiagaraDataInterfaceGeometryCollection : public UNiagaraDataInterface
 	UPROPERTY(EditAnywhere, Category = "Geometry Collection")
 	bool bIncludeIntermediateBones = false;
 
-	UPROPERTY(Transient)
-	FResolvedNiagaraGeometryCollection ResolvedSource;
-
 	/** UObject Interface */
 	virtual void PostInitProperties() override;
 
@@ -287,12 +280,12 @@ protected:
 	virtual bool CopyToInternal(UNiagaraDataInterface* Destination) const override;
 
 private:
-	void ResolveGeometryCollection(FNiagaraSystemInstance* SystemInstance);
-	bool ResolveGeometryCollectionFromDirectSource();
-	bool ResolveGeometryCollectionFromAttachParent(FNiagaraSystemInstance* SystemInstance);
-	bool ResolveGeometryCollectionFromActor(AActor* Actor);
-	bool ResolveGeometryCollectionFromDefaultCollection();
-	bool ResolveGeometryCollectionFromParamterBinding(UObject* UserParameter);
+	void ResolveGeometryCollection(FNiagaraSystemInstance* SystemInstance, FNDIGeometryCollectionData* InstanceData);
+	bool ResolveGeometryCollectionFromDirectSource(FResolvedNiagaraGeometryCollection& ResolvedSource);
+	bool ResolveGeometryCollectionFromAttachParent(FNiagaraSystemInstance* SystemInstance, FResolvedNiagaraGeometryCollection& ResolvedSource);
+	bool ResolveGeometryCollectionFromActor(AActor* Actor, FResolvedNiagaraGeometryCollection& ResolvedSource);
+	bool ResolveGeometryCollectionFromDefaultCollection(FResolvedNiagaraGeometryCollection& ResolvedSource);
+	bool ResolveGeometryCollectionFromParameterBinding(UObject* UserParameter, FResolvedNiagaraGeometryCollection& ResolvedSource);
 };
 
 /** Proxy to send data to gpu */
