Updated geometry collection DI to always return the current transform per element, not just the initial ones

diff --git a/Engine/Plugins/Experimental/ChaosNiagara/Source/ChaosNiagara/Private/NiagaraDataInterfaceGeometryCollection.cpp b/Engine/Plugins/Experimental/ChaosNiagara/Source/ChaosNiagara/Private/NiagaraDataInterfaceGeometryCollection.cpp
--- a/Engine/Plugins/Experimental/ChaosNiagara/Source/ChaosNiagara/Private/NiagaraDataInterfaceGeometryCollection.cpp
+++ b/Engine/Plugins/Experimental/ChaosNiagara/Source/ChaosNiagara/Private/NiagaraDataInterfaceGeometryCollection.cpp
@@ -210,11 +210,11 @@ void FNDIGeometryCollectionData::Update(UNiagaraDataInterfaceGeometryCollection*
 			{
 				Init(Interface, SystemInstance);
 				bNeedsRenderUpdate = true;
-				AssetArrays->ComponentRestTransformBuffer = ResolvedSource.GetInitialLocalRestTransforms();
+				AssetArrays->ComponentRestTransformBuffer = ResolvedSource.GetLocalRestTransforms();
 			}
 			else
 			{
-				TArray<FTransform> NewTransforms = ResolvedSource.GetInitialLocalRestTransforms();
+				TArray<FTransform> NewTransforms = ResolvedSource.GetLocalRestTransforms();
 				int32 TransformCount = NewTransforms.Num();
 				if (TransformCount != AssetArrays->ComponentRestTransformBuffer.Num() ||
 					FMemory::Memcmp(NewTransforms.GetData(), AssetArrays->ComponentRestTransformBuffer.GetData(), TransformCount * sizeof(FTransform)) != 0)
@@ -315,11 +315,11 @@ FTransform FResolvedNiagaraGeometryCollection::GetComponentSpaceTransform(int32
 	return FTransform();
 }
 
-TArray<FTransform> FResolvedNiagaraGeometryCollection::GetInitialLocalRestTransforms() const
+TArray<FTransform> FResolvedNiagaraGeometryCollection::GetLocalRestTransforms() const
 {
 	if (Component.IsValid())
 	{
-		return Component->GetInitialLocalRestTransforms();
+		return Component->GetLocalRestTransforms();
 	}
 	if (Collection.IsValid())
 	{
diff --git a/Engine/Plugins/Experimental/ChaosNiagara/Source/ChaosNiagara/Public/NiagaraDataInterfaceGeometryCollection.h b/Engine/Plugins/Experimental/ChaosNiagara/Source/ChaosNiagara/Public/NiagaraDataInterfaceGeometryCollection.h
--- a/Engine/Plugins/Experimental/ChaosNiagara/Source/ChaosNiagara/Public/NiagaraDataInterfaceGeometryCollection.h
+++ b/Engine/Plugins/Experimental/ChaosNiagara/Source/ChaosNiagara/Public/NiagaraDataInterfaceGeometryCollection.h
@@ -109,7 +109,7 @@ struct FResolvedNiagaraGeometryCollection
 	const UGeometryCollection* GetGeometryCollection() const;
 	FTransform GetComponentRootTransform(FNiagaraSystemInstance* SystemInstance) const;
 	FTransform GetComponentSpaceTransform(int32 TransformIndex) const;
-	TArray<FTransform> GetInitialLocalRestTransforms() const;
+	TArray<FTransform> GetLocalRestTransforms() const;
 	
 	TWeakObjectPtr<UGeometryCollection> Collection;
 	TWeakObjectPtr<UGeometryCollectionComponent> Component;
