* Added element index output to geometry collection DI
* Fixed some minor bugs in the geometry collection DI

#tests Tested that existing assets work with the new function signature + engine tests still run

diff --git a/Engine/Plugins/Experimental/ChaosNiagara/Shaders/NiagaraDataInterfaceGeometryCollection.ush b/Engine/Plugins/Experimental/ChaosNiagara/Shaders/NiagaraDataInterfaceGeometryCollection.ush
--- a/Engine/Plugins/Experimental/ChaosNiagara/Shaders/NiagaraDataInterfaceGeometryCollection.ush
+++ b/Engine/Plugins/Experimental/ChaosNiagara/Shaders/NiagaraDataInterfaceGeometryCollection.ush
@@ -33,37 +33,25 @@ ByteAddressBuffer  {ParameterName}_ElementTransforms;
  * -----------------------------------------------------------------
  */
 
-float3 DIGeometryCollection_GetCurrentPosition(in int ElementIndex)
-{
-	const int BufferOffset = 3 * ElementIndex;
-    return float3({ParameterName}_WorldTransformBuffer[BufferOffset].w, {ParameterName}_WorldTransformBuffer[BufferOffset+1].w, {ParameterName}_WorldTransformBuffer[BufferOffset+2].w);
-}
-
-float3 DIGeometryCollection_GetPreviousPosition(in int ElementIndex)
-{
-	const int BufferOffset = 3 * ElementIndex;
-    return float3({ParameterName}_WorldTransformBuffer[BufferOffset].w, {ParameterName}_WorldTransformBuffer[BufferOffset+1].w, {ParameterName}_WorldTransformBuffer[BufferOffset+2].w);
-}
-
-float3x4 DIGeometryCollection_GetCurrentTransform(in int ElementIndex)
+float3x4 DIGeometryCollection_GetCurrentTransform_{ParameterName}(in int ElementIndex)
 {
 	const int BufferOffset = 3 * ElementIndex;
     return float3x4({ParameterName}_WorldTransformBuffer[BufferOffset], {ParameterName}_WorldTransformBuffer[BufferOffset+1], {ParameterName}_WorldTransformBuffer[BufferOffset+2]);
 }
 
-float3x4 DIGeometryCollection_GetInverseTransform(in int ElementIndex)
+float3x4 DIGeometryCollection_GetInverseTransform_{ParameterName}(in int ElementIndex)
 {
 	const int BufferOffset = 3 * ElementIndex;
     return float3x4({ParameterName}_WorldInverseTransformBuffer[BufferOffset], {ParameterName}_WorldInverseTransformBuffer[BufferOffset+1], {ParameterName}_WorldInverseTransformBuffer[BufferOffset+2]);
 }
 
-float3x4 DIGeometryCollection_GetPreviousTransform(in int ElementIndex)
+float3x4 DIGeometryCollection_GetPreviousTransform_{ParameterName}(in int ElementIndex)
 {
 	const int BufferOffset = 3 * ElementIndex;
     return float3x4({ParameterName}_PrevWorldTransformBuffer[BufferOffset], {ParameterName}_PrevWorldTransformBuffer[BufferOffset+1], {ParameterName}_PrevWorldTransformBuffer[BufferOffset+2]);
 }
 
-float3x4 DIGeometryCollection_GetPreviousInverse(in int ElementIndex)
+float3x4 DIGeometryCollection_GetPreviousInverse_{ParameterName}(in int ElementIndex)
 {
 	const int BufferOffset = 3 * ElementIndex;
     return float3x4({ParameterName}_PrevWorldInverseTransformBuffer[BufferOffset], {ParameterName}_PrevWorldInverseTransformBuffer[BufferOffset+1], {ParameterName}_PrevWorldInverseTransformBuffer[BufferOffset+2]);
@@ -75,37 +63,7 @@ float3x4 DIGeometryCollection_GetPreviousInverse(in int ElementIndex)
  */
 
 // Given a world space position (WorldPosition) compute the sphere closest point (position,normal,velocity)
-float DIGeometryCollection_GetSphereProjection(in float3 LocalPosition, in float3 SphereCenter, in float SphereRadius, in int SphereIndex,
-			inout float3 OutClosestPosition, inout float3 OutClosestNormal, inout int OutElementIndex, inout float OutMinDistance)
-{
-	const float3 DeltaPosition = LocalPosition - SphereCenter;
-	const float DeltaLength = length(DeltaPosition);
-	const float SphereDistance = DeltaLength - SphereRadius;
-
-	if (SphereDistance < OutMinDistance)
-	{
-		OutMinDistance = SphereDistance;
-		OutElementIndex = SphereIndex;
-
-		OutClosestNormal = (DeltaLength > SMALL_NUMBER) ? DeltaPosition / DeltaLength : float3(0,0,0);
-		OutClosestPosition = LocalPosition - OutClosestNormal*SphereDistance;
-	}
-	return SphereDistance;
-}
-
-// Given a world space position (WorldPosition) compute the sphere closest point (position,normal,velocity)
-float DIGeometryCollection_GetCapsuleProjection(in float3 LocalPosition, in float2 CapsuleSize, in int CapsuleIndex,
-			inout float3 OutClosestPosition, inout float3 OutClosestNormal, inout int OutElementIndex, inout float OutMinDistance)
-{
-	const float HalfLength = 0.5 * CapsuleSize.y;
-	const float3 SphereCenter = float3(0,0,clamp(LocalPosition.z, -HalfLength, HalfLength));
-
-	return DIGeometryCollection_GetSphereProjection(LocalPosition,SphereCenter,CapsuleSize.x,CapsuleIndex,
-			OutClosestPosition,OutClosestNormal,OutElementIndex,OutMinDistance);
-}
-
-// Given a world space position (WorldPosition) compute the sphere closest point (position,normal,velocity)
-float DIGeometryCollection_GetBoxProjection(in float3 LocalPosition, in float3 BoxExtent, in int BoxIndex, 
+float DIGeometryCollection_GetBoxProjection_{ParameterName}(in float3 LocalPosition, in float3 BoxExtent, in int BoxIndex, 
 			inout float3 OutClosestPosition, inout float3 OutClosestNormal, inout int OutElementIndex, inout float OutMinDistance)
 {
 	const float3 HalfExtent = 0.5 * BoxExtent;
@@ -199,16 +157,16 @@ void GetElementBounds_{ParameterName}(in int InElementIndex, out float3 OutCente
  * -----------------------------------------------------------------
  */
 
-float3 DIGeometryCollection_GetLocalPosition(in float3 WorldPosition, in int ElementIndex, in float TimeFraction )
+float3 DIGeometryCollection_GetLocalPosition_{ParameterName}(in float3 WorldPosition, in int ElementIndex, in float TimeFraction )
 {
-	const float3 CurrentLocal = mul(DIGeometryCollection_GetInverseTransform(ElementIndex), float4(WorldPosition,1.0)).xyz;
-	const float3 PreviousLocal = mul(DIGeometryCollection_GetPreviousInverse(ElementIndex), float4(WorldPosition,1.0)).xyz;
+	const float3 CurrentLocal = mul(DIGeometryCollection_GetInverseTransform_{ParameterName}(ElementIndex), float4(WorldPosition,1.0)).xyz;
+	const float3 PreviousLocal = mul(DIGeometryCollection_GetPreviousInverse_{ParameterName}(ElementIndex), float4(WorldPosition,1.0)).xyz;
 
 	return PreviousLocal + TimeFraction * (CurrentLocal-PreviousLocal);
 }
 
 // Given a world space position (WorldPosition) compute the static mesh closest point (position,normal,velocity)
-void DIGeometryCollection_GetClosestPointNoNormal(in float3 WorldPosition, in float DeltaTime, in float TimeFraction, out float OutClosestDistance, out float3 OutClosestPosition, out float3 OutClosestVelocity)
+void DIGeometryCollection_GetClosestPointNoNormal_{ParameterName}(in float3 WorldPosition, in float DeltaTime, in float TimeFraction, out float OutClosestDistance, out float3 OutClosestPosition, out float3 OutClosestVelocity, out int ElementIndex)
 {	
 	float3 CollisionPosition = float3(0,0,0);
 	float3 CollisionNormal = float3(0,0,0);
@@ -216,7 +174,7 @@ void DIGeometryCollection_GetClosestPointNoNormal(in float3 WorldPosition, in fl
 	
 	OutClosestPosition = float3(0,0,0);
 	OutClosestVelocity = float3(0,0,0);
-	
+	ElementIndex = -1;
 	
 	[branch]
 	if (WorldPosition.x < {ParameterName}_BoundsMin.x || 
@@ -233,25 +191,26 @@ void DIGeometryCollection_GetClosestPointNoNormal(in float3 WorldPosition, in fl
 	int ClosestTransform = -1;
 	for (int CurrTransform = 0; CurrTransform < {ParameterName}_NumPieces; ++CurrTransform)
 	{
-		float3 LocalPosition = DIGeometryCollection_GetLocalPosition(WorldPosition,CurrTransform,TimeFraction);
+		float3 LocalPosition = DIGeometryCollection_GetLocalPosition_{ParameterName}(WorldPosition,CurrTransform,TimeFraction);
 		
-		DIGeometryCollection_GetBoxProjection(LocalPosition, {ParameterName}_BoundsBuffer[CurrTransform].xyz, CurrTransform,
+		DIGeometryCollection_GetBoxProjection_{ParameterName}(LocalPosition, {ParameterName}_BoundsBuffer[CurrTransform].xyz, CurrTransform,
 				CollisionPosition, CollisionNormal, ClosestTransform, OutClosestDistance);
 	}
 
 	if (ClosestTransform != -1)
 	{
-		const float3 PreviousPosition = mul(DIGeometryCollection_GetPreviousTransform(ClosestTransform), float4(CollisionPosition,1.0)).xyz;
-		const float3 CurrentPosition = mul(DIGeometryCollection_GetCurrentTransform(ClosestTransform), float4(CollisionPosition,1.0)).xyz;
+		const float3 PreviousPosition = mul(DIGeometryCollection_GetPreviousTransform_{ParameterName}(ClosestTransform), float4(CollisionPosition,1.0)).xyz;
+		const float3 CurrentPosition = mul(DIGeometryCollection_GetCurrentTransform_{ParameterName}(ClosestTransform), float4(CollisionPosition,1.0)).xyz;
 		OutClosestVelocity = ( CurrentPosition - PreviousPosition ) / DeltaTime;
-		OutClosestPosition = PreviousPosition + TimeFraction * (CurrentPosition-PreviousPosition);		
+		OutClosestPosition = PreviousPosition + TimeFraction * (CurrentPosition-PreviousPosition);
+		ElementIndex = ClosestTransform;
 	}
 }
 
 
-void GetClosestPointNoNormal_{ParameterName}(in float3 WorldPosition, in float DeltaTime, in float TimeFraction, out float ClosestDistance, out float3 OutClosestPosition, out float3 OutClosestVelocity)
+void GetClosestPointNoNormal_{ParameterName}(in float3 WorldPosition, in float DeltaTime, in float TimeFraction, out float ClosestDistance, out float3 OutClosestPosition, out float3 OutClosestVelocity, out int ElementIndex)
 {
-	DIGeometryCollection_GetClosestPointNoNormal(WorldPosition,DeltaTime, TimeFraction, ClosestDistance, OutClosestPosition, OutClosestVelocity);
+	DIGeometryCollection_GetClosestPointNoNormal_{ParameterName}(WorldPosition,DeltaTime, TimeFraction, ClosestDistance, OutClosestPosition, OutClosestVelocity, ElementIndex);
 }
 
 #undef BOX_INDEX
diff --git a/Engine/Plugins/Experimental/ChaosNiagara/Source/ChaosNiagara/Private/NiagaraDataInterfaceGeometryCollection.cpp b/Engine/Plugins/Experimental/ChaosNiagara/Source/ChaosNiagara/Private/NiagaraDataInterfaceGeometryCollection.cpp
--- a/Engine/Plugins/Experimental/ChaosNiagara/Source/ChaosNiagara/Private/NiagaraDataInterfaceGeometryCollection.cpp
+++ b/Engine/Plugins/Experimental/ChaosNiagara/Source/ChaosNiagara/Private/NiagaraDataInterfaceGeometryCollection.cpp
@@ -66,6 +66,7 @@ namespace NDIGeometryCollectionLocal
 		{
 			InitialVersion = 0,
 			DIRefactor = 1,
+			AddedElementIndexOutput = 2,
 			
 			VersionPlusOne,
 			LatestVersion = VersionPlusOne - 1
@@ -231,6 +232,11 @@ void FNDIGeometryCollectionData::Update(UNiagaraDataInterfaceGeometryCollection*
 			{
 				ResolvedSource.Component->Bounds.GetBox().GetCenterAndExtents(Origin, Extents);
 			}
+			else
+			{
+				// if the extent is 0 then some functions won't work with the preview collection, as it doesn't have a component
+				Extents = FVector::OneVector * UE_FLOAT_HUGE_DISTANCE;
+			}
 
 			FNiagaraLWCConverter LwcConverter = SystemInstance->GetLWCConverter();
 			BoundsOrigin = LwcConverter.ConvertWorldToSimulationVector(Origin);
@@ -696,6 +702,7 @@ void UNiagaraDataInterfaceGeometryCollection::GetFunctionsInternal(TArray<FNiaga
 		Sig.Outputs.Add(FNiagaraVariable(FNiagaraTypeDefinition::GetFloatDef(), TEXT("Closest Distance")));
 		Sig.Outputs.Add(FNiagaraVariable(FNiagaraTypeDefinition::GetPositionDef(), TEXT("Closest Position")));
 		Sig.Outputs.Add(FNiagaraVariable(FNiagaraTypeDefinition::GetVec3Def(), TEXT("Closest Velocity")));
+		Sig.Outputs.Add(FNiagaraVariable(FNiagaraTypeDefinition::GetIntDef(), TEXT("Element Index")));
 		NIAGARA_ADD_FUNCTION_SOURCE_INFO(Sig)
 
 		OutFunctions.Add(Sig);
@@ -1009,7 +1016,7 @@ bool UNiagaraDataInterfaceGeometryCollection::GetFunctionHLSL(const FNiagaraData
 bool UNiagaraDataInterfaceGeometryCollection::UpgradeFunctionCall(FNiagaraFunctionSignature& FunctionSignature)
 {
 	// upgrade old functions to the latest version
-	if (FunctionSignature.FunctionVersion < NDIGeometryCollectionLocal::FGeometryCollectionDIFunctionVersion::DIRefactor)
+	if (FunctionSignature.FunctionVersion < NDIGeometryCollectionLocal::FGeometryCollectionDIFunctionVersion::AddedElementIndexOutput)
 	{
 		TArray<FNiagaraFunctionSignature> AllFunctions;
 		GetFunctionsInternal(AllFunctions);
