Chaos Niagara interface : Perf / memory and trailing event not play from a chaos cache fixes
- fix issue where the interface may accumulate breaking events but never reset them causing performance and memory issue
- Moved the BreakingEvents array reset to when the breaking events are handled ( to mimic the logic used for other event types)
- Only register events handlers that match the interface event type
- Rewrote the trailing event handler to properly read the input data instead of infering the data from the dyanmic particle of the geometry collection component
- removed TrailingGeometryCollectionComponents as it is no loinger needed
- fixed missing initialization of TransformGroupIndex in FTrailingData default constructor

diff --git a/Engine/Plugins/Experimental/ChaosNiagara/Source/ChaosNiagara/Private/NiagaraDataInterfaceChaosDestruction.cpp b/Engine/Plugins/Experimental/ChaosNiagara/Source/ChaosNiagara/Private/NiagaraDataInterfaceChaosDestruction.cpp
--- a/Engine/Plugins/Experimental/ChaosNiagara/Source/ChaosNiagara/Private/NiagaraDataInterfaceChaosDestruction.cpp
+++ b/Engine/Plugins/Experimental/ChaosNiagara/Source/ChaosNiagara/Private/NiagaraDataInterfaceChaosDestruction.cpp
@@ -579,12 +579,23 @@ void UNiagaraDataInterfaceChaosDestruction::RegisterWithSolverEventManager(Chaos
 	Chaos::FEventManager* EventManager = Solver->GetEventManager();
 	if (EventManager)
 	{
-		Solver->SetGenerateCollisionData(true);
-		Solver->SetGenerateBreakingData(true);
-		Solver->SetGenerateTrailingData(true);
-		EventManager->RegisterHandler<Chaos::FCollisionEventData>(Chaos::EEventType::Collision, this, &UNiagaraDataInterfaceChaosDestruction::HandleCollisionEvents);
-		EventManager->RegisterHandler<Chaos::FBreakingEventData>(Chaos::EEventType::Breaking, this, &UNiagaraDataInterfaceChaosDestruction::HandleBreakingEvents);
-		EventManager->RegisterHandler<Chaos::FTrailingEventData>(Chaos::EEventType::Trailing, this, &UNiagaraDataInterfaceChaosDestruction::HandleTrailingEvents);
+		switch (DataSourceType)
+		{
+		case EDataSourceTypeEnum::ChaosNiagara_DataSourceType_Collision:
+			Solver->SetGenerateCollisionData(true);
+			EventManager->RegisterHandler<Chaos::FCollisionEventData>(Chaos::EEventType::Collision, this, &UNiagaraDataInterfaceChaosDestruction::HandleCollisionEvents);
+			break;
+		case EDataSourceTypeEnum::ChaosNiagara_DataSourceType_Breaking:
+			Solver->SetGenerateBreakingData(true);
+			EventManager->RegisterHandler<Chaos::FBreakingEventData>(Chaos::EEventType::Breaking, this, &UNiagaraDataInterfaceChaosDestruction::HandleBreakingEvents);
+			break;
+		case EDataSourceTypeEnum::ChaosNiagara_DataSourceType_Trailing:
+			Solver->SetGenerateBreakingData(true);
+			EventManager->RegisterHandler<Chaos::FBreakingEventData>(Chaos::EEventType::Breaking, this, &UNiagaraDataInterfaceChaosDestruction::HandleBreakingEvents);
+			Solver->SetGenerateTrailingData(true);
+			EventManager->RegisterHandler<Chaos::FTrailingEventData>(Chaos::EEventType::Trailing, this, &UNiagaraDataInterfaceChaosDestruction::HandleTrailingEvents);
+			break;
+		}
 	}
 }
 
@@ -1440,8 +1451,10 @@ void UNiagaraDataInterfaceChaosDestruction::HandleBreakingEvents(const Chaos::FB
 	ensure(IsInGameThread());
 	Chaos::FBreakingDataArray const& BreakingDataIn = Event.BreakingData.AllBreakingsArray;
 
+	BreakingEvents.Reset();
+
 	// Copy data from Event
-	BreakingEvents.InsertZeroed(0, BreakingDataIn.Num());
+	BreakingEvents.AddZeroed(BreakingDataIn.Num());
 
 	int32 Idx = 0;
 	for (Chaos::FBreakingData const& DataIn : BreakingDataIn)
@@ -1499,11 +1512,6 @@ void UNiagaraDataInterfaceChaosDestruction::HandleBreakingEvents(const Chaos::FB
 				{
 					CopyData.PhysicalMaterialName = FName();
 				}
-
-				if (GeometryCollectionComponent->GetNotifyTrailing())
-				{
-					TrailingGeometryCollectionComponents.Add(GeometryCollectionComponent);
-				}
 			}
 		}
 		else
@@ -1986,8 +1994,6 @@ bool UNiagaraDataInterfaceChaosDestruction::BreakingCallback(FNDIChaosDestructio
 		IdxSolver++;
 	}
 
-	BreakingEvents.Reset();
-
 	INC_DWORD_STAT_BY(STAT_NiagaraNumParticlesSpawnedFromBreakings, InstData->PositionArray.Num());
 
 	return false;
@@ -2000,8 +2006,7 @@ void UNiagaraDataInterfaceChaosDestruction::HandleTrailingEvents(const Chaos::FT
 
 	TrailingEvents.Reset();
 
-	float TrailMinSpeedToSpawnSquared = TrailMinSpeedToSpawn * TrailMinSpeedToSpawn;
-	UPhysicalMaterial* PhysicalMaterial = nullptr;
+	const float TrailMinSpeedToSpawnSquared = TrailMinSpeedToSpawn * TrailMinSpeedToSpawn;
 
 	auto IsMaterialInFilter = [&](const FName& InMaterialName) 
 	{
@@ -2026,92 +2031,52 @@ void UNiagaraDataInterfaceChaosDestruction::HandleTrailingEvents(const Chaos::FT
 		return false;
 	};
 
-	TArray<TWeakObjectPtr<UGeometryCollectionComponent>> ComponentsToRemove;
-	for (const TWeakObjectPtr<UGeometryCollectionComponent>& WeakGeometryCollectionComponent: TrailingGeometryCollectionComponents)
-	{
-		if (!WeakGeometryCollectionComponent.IsValid())
-		{
-			ComponentsToRemove.Add(WeakGeometryCollectionComponent);
-			continue;
-		}
-
-		if (const TStrongObjectPtr<UGeometryCollectionComponent> GeometryCollectionComponent = WeakGeometryCollectionComponent.Pin())
+	auto GetGeometryCollectionComponentFromProxy = 
+		[this](IPhysicsProxyBase* InProxy) -> UGeometryCollectionComponent*
 		{
-			if (!GeometryCollectionComponent->IsPhysicsStateCreated())
-			{
-				ComponentsToRemove.Add(WeakGeometryCollectionComponent);
-			}
-			else if (GeometryCollectionComponent->GetNotifyTrailing())
+			for (auto& Solver : Solvers)
 			{
-				PhysicalMaterial = GeometryCollectionComponent->GetPhysicalMaterial();
-				ensure(PhysicalMaterial);
-				if (PhysicalMaterial)
+				if (Solver.PhysScene)
 				{
-					if (bApplyMaterialsFilter && !IsMaterialInFilter(PhysicalMaterial->GetFName()))
+					if (UGeometryCollectionComponent* GeometryCollectionComponent = Solver.PhysScene->GetOwningComponent<UGeometryCollectionComponent>(InProxy))
 					{
-						continue;
+						return GeometryCollectionComponent;
 					}
 				}
+			}
+			return nullptr;
+		};
 
-				const TArray<FTransform3f>& ComponentSpaceTransforms = GeometryCollectionComponent->GetComponentSpaceTransforms3f();
-				const FTransform ActorTransform = GeometryCollectionComponent->GetComponentToWorld();
-
-				const FGeometryDynamicCollection* DynamicCollection = GeometryCollectionComponent->GetDynamicCollection();
+	const Chaos::FTrailingDataArray& InEvents = Event.TrailingData.AllTrailingsArray;
 
-				if (DynamicCollection)
+	for (const Chaos::FTrailingData& InEvent : InEvents)
+	{
+		if (InEvent.Velocity.SquaredLength() >= TrailMinSpeedToSpawnSquared)
+		{
+			if (UGeometryCollectionComponent* GeometryCollectionComponent = GetGeometryCollectionComponentFromProxy(InEvent.Proxy))
+			{
+				if (GeometryCollectionComponent->GetNotifyTrailing())
 				{
-					if (const UGeometryCollection* RestCollection = GeometryCollectionComponent->GetRestCollection())
+					if (UPhysicalMaterial* PhysicalMaterial = GeometryCollectionComponent->GetPhysicalMaterial())
 					{
-						const TSharedPtr<FGeometryCollection, ESPMode::ThreadSafe> GeometryCollection = RestCollection->GetGeometryCollection();
-
-						// Get the MassToLocal transforms
-						if (GeometryCollection->HasAttribute(TEXT("MassToLocal"), FTransformCollection::TransformGroup))
+						if (!bApplyMaterialsFilter || IsMaterialInFilter(PhysicalMaterial->GetFName()))
 						{
-							const TManagedArray<FTransform>& CollectionMassToLocal = GeometryCollection->GetAttribute<FTransform>(TEXT("MassToLocal"), FTransformCollection::TransformGroup);
-
-							const TManagedArray<FVector3f>* LinearVelocity = DynamicCollection->FindAttributeTyped<FVector3f>("LinearVelocity", FTransformCollection::TransformGroup);
-							const TManagedArray<FVector3f>* AngularVelocity = DynamicCollection->FindAttributeTyped<FVector3f>("AngularVelocity", FTransformCollection::TransformGroup);
-
-							if (!LinearVelocity || !AngularVelocity)
-							{
-								continue;
-							}
+							Chaos::FTrailingDataExt TrailingData(InEvent);
 
-							ensure(DynamicCollection->Active.Num() == (*LinearVelocity).Num());
-							ensure(DynamicCollection->Active.Num() == (*AngularVelocity).Num());
+							TrailingData.SurfaceType = PhysicalMaterial->SurfaceType;
+							const FBox BoundingBox(InEvent.BoundingBox.Min(), InEvent.BoundingBox.Max());
+							const FVector Extents = BoundingBox.GetSize();
+							TrailingData.BoundingboxVolume = BoundingBox.GetVolume();
+							TrailingData.BoundingboxExtentMin = FMath::Min3(Extents[0], Extents[1], Extents[2]);
+							TrailingData.BoundingboxExtentMax = FMath::Max3(Extents[0], Extents[1], Extents[2]);
 
-							for (int Idx = 0; Idx < DynamicCollection->Active.Num(); ++Idx)
-							{
-								if (DynamicCollection->Active[Idx] && (*LinearVelocity)[Idx].SquaredLength() >= TrailMinSpeedToSpawnSquared)
-								{
-									Chaos::FTrailingDataExt TrailingData;
-
-									const FTransform CurrTransform = CollectionMassToLocal[Idx] * FTransform(ComponentSpaceTransforms[Idx]) * ActorTransform;
-									TrailingData.Location = CurrTransform.GetTranslation();
-
-									TrailingData.Velocity = (*LinearVelocity)[Idx];
-									TrailingData.AngularVelocity = (*AngularVelocity)[Idx];
-
-									TrailingData.Mass = 1.f;
-									TrailingData.BoundingboxVolume = 1000000.f;
-									TrailingData.BoundingboxExtentMin = 100.f;
-									TrailingData.BoundingboxExtentMax = 100.f;
-									TrailingData.SurfaceType = 0;
-
-									TrailingEvents.Add(TrailingData);
-								}
-							}
+							TrailingEvents.Add(TrailingData);
 						}
 					}
-				}				
+				}
 			}
 		}
 	}
-
-	for (const TWeakObjectPtr<UGeometryCollectionComponent>& ComponentToRemove : ComponentsToRemove)
-	{
-		TrailingGeometryCollectionComponents.Remove(ComponentToRemove);
-	}
 }
 
 void UNiagaraDataInterfaceChaosDestruction::FilterAllTrailings(TArray<Chaos::FTrailingDataExt>& AllTrailingsArray)
@@ -2514,7 +2479,8 @@ bool UNiagaraDataInterfaceChaosDestruction::PerInstanceTick(void* PerInstanceDat
 
 	ResetInstData(InstData);
 
-	if (ShouldSpawn && DoSpawn) {
+	if (ShouldSpawn && DoSpawn) 
+	{
 		if (DataSourceType == EDataSourceTypeEnum::ChaosNiagara_DataSourceType_Collision)
 		{
 			SCOPE_CYCLE_COUNTER(STAT_CollisionCallback);
diff --git a/Engine/Source/Runtime/Experimental/Chaos/Public/Chaos/ExternalCollisionData.h b/Engine/Source/Runtime/Experimental/Chaos/Public/Chaos/ExternalCollisionData.h
--- a/Engine/Source/Runtime/Experimental/Chaos/Public/Chaos/ExternalCollisionData.h
+++ b/Engine/Source/Runtime/Experimental/Chaos/Public/Chaos/ExternalCollisionData.h
@@ -354,6 +354,7 @@ namespace Chaos
 			, Mass((FReal)0.0)
 			, Proxy(nullptr)
 			, BoundingBox(FAABB3(FVec3((FReal)0.0), FVec3((FReal)0.0)))
+			, TransformGroupIndex(INDEX_NONE)
 		{}
 
 		UE_DEPRECATED(5.6, "Deprecated constructor")
