Fix wrong indexing in Niagara  Destruction Data interface causing some particles to spawn at the origin

diff --git a/Engine/Plugins/Experimental/ChaosNiagara/Source/ChaosNiagara/Private/NiagaraDataInterfaceChaosDestruction.cpp b/Engine/Plugins/Experimental/ChaosNiagara/Source/ChaosNiagara/Private/NiagaraDataInterfaceChaosDestruction.cpp
--- a/Engine/Plugins/Experimental/ChaosNiagara/Source/ChaosNiagara/Private/NiagaraDataInterfaceChaosDestruction.cpp
+++ b/Engine/Plugins/Experimental/ChaosNiagara/Source/ChaosNiagara/Private/NiagaraDataInterfaceChaosDestruction.cpp
@@ -2951,7 +2951,7 @@ void UNiagaraDataInterfaceChaosDestruction::GetPosition(FVectorVMExternalFunctio
 
 	for (int32 i = 0; i < Context.GetNumInstances(); ++i)
 	{
-		const int32 ParticleID = ParticleIDParam.GetAndAdvance() - LastSpawnedPointID + 1;
+		const int32 ParticleID = ParticleIDParam.GetAndAdvance() - (LastSpawnedPointID + 1);
 		//-TODO: LWC
 		const FVector3f Sample = InstData->PositionArray.IsValidIndex(ParticleID) ? (FVector3f)InstData->PositionArray[ParticleID] : FVector3f::Zero();
 		OutSample.SetAndAdvance(Sample);
@@ -2967,7 +2967,7 @@ void UNiagaraDataInterfaceChaosDestruction::GetNormal(FVectorVMExternalFunctionC
 
 	for (int32 i = 0; i < Context.GetNumInstances(); ++i)
 	{
-		const int32 ParticleID = ParticleIDParam.GetAndAdvance() - LastSpawnedPointID + 1;
+		const int32 ParticleID = ParticleIDParam.GetAndAdvance() - (LastSpawnedPointID + 1);
 		const FVector3f Sample = InstData->IncomingNormalArray.IsValidIndex(ParticleID) ? (FVector3f)InstData->IncomingNormalArray[ParticleID] : FVector3f::Zero();
 		OutSample.SetAndAdvance(Sample);
 	}
@@ -2982,7 +2982,7 @@ void UNiagaraDataInterfaceChaosDestruction::GetVelocity(FVectorVMExternalFunctio
 
 	for (int32 i = 0; i < Context.GetNumInstances(); ++i)
 	{
-		const int32 ParticleID = ParticleIDParam.GetAndAdvance() - LastSpawnedPointID + 1;
+		const int32 ParticleID = ParticleIDParam.GetAndAdvance() - (LastSpawnedPointID + 1);
 		const FVector3f Sample = InstData->VelocityArray.IsValidIndex(ParticleID) ? (FVector3f)InstData->VelocityArray[ParticleID] : FVector3f::Zero();
 		OutSample.SetAndAdvance(Sample);
 	}
@@ -2997,7 +2997,7 @@ void UNiagaraDataInterfaceChaosDestruction::GetAngularVelocity(FVectorVMExternal
 
 	for (int32 i = 0; i < Context.GetNumInstances(); ++i)
 	{
-		const int32 ParticleID = ParticleIDParam.GetAndAdvance() - LastSpawnedPointID + 1;
+		const int32 ParticleID = ParticleIDParam.GetAndAdvance() - (LastSpawnedPointID + 1);
 		const FVector3f Sample = InstData->IncomingAngularVelocity1Array.IsValidIndex(ParticleID) ? (FVector3f)InstData->IncomingAngularVelocity1Array[ParticleID] : FVector3f::Zero();
 		OutSample.SetAndAdvance(Sample);
 	}
@@ -3012,7 +3012,7 @@ void UNiagaraDataInterfaceChaosDestruction::GetExtentMin(FVectorVMExternalFuncti
 
 	for (int32 i = 0; i < Context.GetNumInstances(); ++i)
 	{
-		const int32 ParticleID = ParticleIDParam.GetAndAdvance() - LastSpawnedPointID + 1;
+		const int32 ParticleID = ParticleIDParam.GetAndAdvance() - (LastSpawnedPointID + 1);
 		const float Sample = InstData->ExtentMinArray.IsValidIndex(ParticleID) ? InstData->ExtentMinArray[ParticleID] : 0.0f;
 		OutSample.SetAndAdvance(Sample);
 	}
@@ -3027,7 +3027,7 @@ void UNiagaraDataInterfaceChaosDestruction::GetExtentMax(FVectorVMExternalFuncti
 
 	for (int32 i = 0; i < Context.GetNumInstances(); ++i)
 	{
-		const int32 ParticleID = ParticleIDParam.GetAndAdvance() - LastSpawnedPointID + 1;
+		const int32 ParticleID = ParticleIDParam.GetAndAdvance() - (LastSpawnedPointID + 1);
 		const float Sample = InstData->ExtentMaxArray.IsValidIndex(ParticleID) ? InstData->ExtentMaxArray[ParticleID] : 0.0f;
 		OutSample.SetAndAdvance(Sample);
 	}
@@ -3042,7 +3042,7 @@ void UNiagaraDataInterfaceChaosDestruction::GetVolume(FVectorVMExternalFunctionC
 
 	for (int32 i = 0; i < Context.GetNumInstances(); ++i)
 	{
-		const int32 ParticleID = ParticleIDParam.GetAndAdvance() - LastSpawnedPointID + 1;
+		const int32 ParticleID = ParticleIDParam.GetAndAdvance() - (LastSpawnedPointID + 1);
 		const float Sample = InstData->VolumeArray.IsValidIndex(ParticleID) ? InstData->VolumeArray[ParticleID] : 0.0f;
 		OutSample.SetAndAdvance(Sample);
 	}
@@ -3110,7 +3110,7 @@ void UNiagaraDataInterfaceChaosDestruction::GetColor(FVectorVMExternalFunctionCo
 
 	for (int32 i = 0; i < Context.GetNumInstances(); ++i)
 	{
-		const int32 ParticleID = ParticleIDParam.GetAndAdvance() - LastSpawnedPointID + 1;
+		const int32 ParticleID = ParticleIDParam.GetAndAdvance() - (LastSpawnedPointID + 1);
 		const FLinearColor Sample = InstData->ColorArray.IsValidIndex(ParticleID) ? InstData->ColorArray[ParticleID] : FLinearColor::White;
 		OutSample.SetAndAdvance(Sample);
 	}
@@ -3137,7 +3137,7 @@ void UNiagaraDataInterfaceChaosDestruction::GetDensity(FVectorVMExternalFunction
 
 	for (int32 i = 0; i < Context.GetNumInstances(); ++i)
 	{
-		const int32 ParticleID = ParticleIDParam.GetAndAdvance() - LastSpawnedPointID + 1;
+		const int32 ParticleID = ParticleIDParam.GetAndAdvance() - (LastSpawnedPointID + 1);
 		const float Sample = InstData->DensityArray.IsValidIndex(ParticleID) ? InstData->DensityArray[ParticleID] : 0.0f;
 		OutSample.SetAndAdvance(Sample);
 	}
@@ -3152,7 +3152,7 @@ void UNiagaraDataInterfaceChaosDestruction::GetFriction(FVectorVMExternalFunctio
 
 	for (int32 i = 0; i < Context.GetNumInstances(); ++i)
 	{
-		const int32 ParticleID = ParticleIDParam.GetAndAdvance() - LastSpawnedPointID + 1;
+		const int32 ParticleID = ParticleIDParam.GetAndAdvance() - (LastSpawnedPointID + 1);
 		const float Sample = InstData->FrictionArray.IsValidIndex(ParticleID) ? InstData->FrictionArray[ParticleID] : 0.0f;
 		OutSample.SetAndAdvance(Sample);
 	}
@@ -3167,7 +3167,7 @@ void UNiagaraDataInterfaceChaosDestruction::GetRestitution(FVectorVMExternalFunc
 
 	for (int32 i = 0; i < Context.GetNumInstances(); ++i)
 	{
-		const int32 ParticleID = ParticleIDParam.GetAndAdvance() - LastSpawnedPointID + 1;
+		const int32 ParticleID = ParticleIDParam.GetAndAdvance() - (LastSpawnedPointID + 1);
 		const float Sample = InstData->RestitutionArray.IsValidIndex(ParticleID) ? InstData->RestitutionArray[ParticleID] : 0.0f;
 		OutSample.SetAndAdvance(Sample);
 	}
@@ -3184,7 +3184,7 @@ void UNiagaraDataInterfaceChaosDestruction::GetTransform(FVectorVMExternalFuncti
 
 	for (int32 i = 0; i < Context.GetNumInstances(); ++i)
 	{
-		const int32 ParticleID = ParticleIDParam.GetAndAdvance() - LastSpawnedPointID + 1;
+		const int32 ParticleID = ParticleIDParam.GetAndAdvance() - (LastSpawnedPointID + 1);
 		if (InstData->TransformTranslationArray.IsValidIndex(ParticleID))
 		{
 			//-TODO: LWC
@@ -3210,7 +3210,7 @@ void UNiagaraDataInterfaceChaosDestruction::GetSize(FVectorVMExternalFunctionCon
 
 	for (int32 i = 0; i < Context.GetNumInstances(); ++i)
 	{
-		const int32 ParticleID = ParticleIDParam.GetAndAdvance() - LastSpawnedPointID + 1;
+		const int32 ParticleID = ParticleIDParam.GetAndAdvance() - (LastSpawnedPointID + 1);
 		const FVector3f Size = InstData->BoundsArray.IsValidIndex(ParticleID) ? (FVector3f)InstData->BoundsArray[ParticleID] : FVector3f::Zero();
 		OutSize.SetAndAdvance(Size);
 	}
@@ -3225,7 +3225,7 @@ void UNiagaraDataInterfaceChaosDestruction::GetSurfaceType(FVectorVMExternalFunc
 
 	for (int32 i = 0; i < Context.GetNumInstances(); ++i)
 	{
-		const int32 ParticleID = ParticleIDParam.GetAndAdvance() - LastSpawnedPointID + 1;
+		const int32 ParticleID = ParticleIDParam.GetAndAdvance() - (LastSpawnedPointID + 1);
 		const int32 Value = InstData->SurfaceTypeArray.IsValidIndex(ParticleID) ? InstData->SurfaceTypeArray[ParticleID] : 0;
 		OutValue.SetAndAdvance(Value);
 	}
@@ -3249,7 +3249,7 @@ void UNiagaraDataInterfaceChaosDestruction::GetCollisionData(FVectorVMExternalFu
 
 	for (int32 i = 0; i < Context.GetNumInstances(); ++i)
 	{
-		const int32 ParticleID = ParticleIDParam.GetAndAdvance() - LastSpawnedPointID + 1;
+		const int32 ParticleID = ParticleIDParam.GetAndAdvance() - (LastSpawnedPointID + 1);
 
 		//-TODO: LWC
 		OutLocation.SetAndAdvance(InstData->IncomingLocationArray.IsValidIndex(ParticleID) ? (FVector3f)InstData->IncomingLocationArray[ParticleID] : FVector3f::Zero());
@@ -3278,7 +3278,7 @@ void UNiagaraDataInterfaceChaosDestruction::GetBreakingData(FVectorVMExternalFun
 
 	for (int32 i = 0; i < Context.GetNumInstances(); ++i)
 	{
-		const int32 ParticleID = ParticleIDParam.GetAndAdvance() - LastSpawnedPointID + 1;
+		const int32 ParticleID = ParticleIDParam.GetAndAdvance() - (LastSpawnedPointID + 1);
 
 		OutLocation.SetAndAdvance(InstData->IncomingLocationArray.IsValidIndex(ParticleID) ? (FVector3f)InstData->IncomingLocationArray[ParticleID] : FVector3f::Zero());
 		OutVelocity.SetAndAdvance(InstData->IncomingVelocity1Array.IsValidIndex(ParticleID) ? (FVector3f)InstData->IncomingVelocity1Array[ParticleID] : FVector3f::Zero());
@@ -3301,7 +3301,7 @@ void UNiagaraDataInterfaceChaosDestruction::GetTrailingData(FVectorVMExternalFun
 
 	for (int32 i = 0; i < Context.GetNumInstances(); ++i)
 	{
-		const int32 ParticleID = ParticleIDParam.GetAndAdvance() - LastSpawnedPointID + 1;
+		const int32 ParticleID = ParticleIDParam.GetAndAdvance() - (LastSpawnedPointID + 1);
 
 		OutLocation.SetAndAdvance(InstData->IncomingLocationArray.IsValidIndex(ParticleID) ? (FVector3f)InstData->IncomingLocationArray[ParticleID] : FVector3f::Zero());
 		OutVelocity.SetAndAdvance(InstData->IncomingVelocity1Array.IsValidIndex(ParticleID) ? (FVector3f)InstData->IncomingVelocity1Array[ParticleID] : FVector3f::Zero());
