Fix debug transient physics field accumulation

#Synced-CL 51879956

diff --git a/Engine/Plugins/Experimental/ChaosNiagara/Source/ChaosNiagara/Classes/NiagaraDataInterfacePhysicsField.h b/Engine/Plugins/Experimental/ChaosNiagara/Source/ChaosNiagara/Classes/NiagaraDataInterfacePhysicsField.h
--- a/Engine/Plugins/Experimental/ChaosNiagara/Source/ChaosNiagara/Classes/NiagaraDataInterfacePhysicsField.h
+++ b/Engine/Plugins/Experimental/ChaosNiagara/Source/ChaosNiagara/Classes/NiagaraDataInterfacePhysicsField.h
@@ -65,7 +65,7 @@ public:
 	UE_API virtual bool PerInstanceTickPostSimulate(void* PerInstanceData, FNiagaraSystemInstance* SystemInstance, float DeltaSeconds) override;
 	virtual int32 PerInstanceDataSize()const override { return sizeof(FNDIPhysicsFieldData); }
 	UE_API virtual bool Equals(const UNiagaraDataInterface* Other) const override;
-	virtual bool HasPostSimulateTick() const override { return true; }
+	virtual bool HasPostSimulateTick() const override { return IsUsedWithCPUScript(); }
 
 	/** GPU simulation functionality */
 #if WITH_EDITORONLY_DATA
diff --git a/Engine/Source/Runtime/Engine/Private/PhysicsField/PhysicsFieldComponent.cpp b/Engine/Source/Runtime/Engine/Private/PhysicsField/PhysicsFieldComponent.cpp
--- a/Engine/Source/Runtime/Engine/Private/PhysicsField/PhysicsFieldComponent.cpp
+++ b/Engine/Source/Runtime/Engine/Private/PhysicsField/PhysicsFieldComponent.cpp
@@ -874,7 +874,7 @@ void UPhysicsFieldComponent::SendRenderDynamicData_Concurrent()
 			LocalInstance->FieldCommands.Append(TransientCommands[FieldBuffers[FieldIndex]]);
 			LocalInstance->FieldCommands.Append(ConstructionCommands[FieldBuffers[FieldIndex]]);
 			
-			if (bIsDebugBuffer)
+			if (bIsDebugBuffer && TransientLifetime > 0.0f)
 			{
 				for (int32 CommandIndex = TransientCommands[FieldBuffers[FieldIndex]].Num() - 1; CommandIndex >= 0; --CommandIndex)
 				{
diff --git a/Engine/Source/Runtime/Renderer/Private/PhysicsFieldRendering.cpp b/Engine/Source/Runtime/Renderer/Private/PhysicsFieldRendering.cpp
--- a/Engine/Source/Runtime/Renderer/Private/PhysicsFieldRendering.cpp
+++ b/Engine/Source/Runtime/Renderer/Private/PhysicsFieldRendering.cpp
@@ -40,7 +40,7 @@ FAutoConsoleVariableRef CVarPhysicsFieldSystemType(
 	TEXT("Physics field boolean to check if we want to display the CPU(0) or GPU(1) field.\n"),
 	ECVF_RenderThreadSafe);
 
-float GPhysicsFieldTransientLifetime = 3.0f;
+float GPhysicsFieldTransientLifetime = 0.0f;
 FAutoConsoleVariableRef CVarPhysicsFieldTransientLifetime(
 	TEXT("r.PhysicsField.Rendering.TransientLifetime"),
 	GPhysicsFieldTransientLifetime,
